Fix Mac build, TRUE and FALSE are defined on that platform

diff --git a/Engine/Plugins/Runtime/GeoReferencing/Source/GeoReferencing/Private/UFSProjSupport.cpp b/Engine/Plugins/Runtime/GeoReferencing/Source/GeoReferencing/Private/UFSProjSupport.cpp
--- a/Engine/Plugins/Runtime/GeoReferencing/Source/GeoReferencing/Private/UFSProjSupport.cpp
+++ b/Engine/Plugins/Runtime/GeoReferencing/Source/GeoReferencing/Private/UFSProjSupport.cpp
@@ -29,8 +29,14 @@ PROJ_FILE_API FUFSProj::FunctionTable =
 	&FUFSProj::Rename,
 };
 
-static constexpr int TRUE = 1;
-static constexpr int FALSE = 0;
+#ifndef TRUE
+#define TRUE (1)
+#endif
+
+#ifndef FALSE
+#define FALSE (0)
+#endif
+
 static const bool EnableLogging = true;
 
 /** Open file. Return NULL if error */
