PR #8371: add GeometricObjects and Engine module to GeoReferencing's PublicDependencyModuleNames. (Contributed by [REDACTED_PERSON])

Just some ws and alpha ordering fixup by GE

diff --git a/Engine/Plugins/Runtime/GeoReferencing/Source/GeoReferencing/GeoReferencing.Build.cs b/Engine/Plugins/Runtime/GeoReferencing/Source/GeoReferencing/GeoReferencing.Build.cs
--- a/Engine/Plugins/Runtime/GeoReferencing/Source/GeoReferencing/GeoReferencing.Build.cs
+++ b/Engine/Plugins/Runtime/GeoReferencing/Source/GeoReferencing/GeoReferencing.Build.cs
@@ -16,6 +16,8 @@ public class GeoReferencing : ModuleRules
 			new string[]
 			{
 				"Core"
+				"Engine",
+				"GeometricObjects",
 			}
 		);
 
