Undo [REDACTED_HANDLE]

diff --git a/Engine/Plugins/Runtime/GeoReferencing/Source/GeoReferencing/GeoReferencing.Build.cs b/Engine/Plugins/Runtime/GeoReferencing/Source/GeoReferencing/GeoReferencing.Build.cs
--- a/Engine/Plugins/Runtime/GeoReferencing/Source/GeoReferencing/GeoReferencing.Build.cs
+++ b/Engine/Plugins/Runtime/GeoReferencing/Source/GeoReferencing/GeoReferencing.Build.cs
@@ -16,8 +16,6 @@ public class GeoReferencing : ModuleRules
 			new string[]
 			{
 				"Core"
-				"Engine",
-				"GeometricObjects",
 			}
 		);
 
