PR #8371: add GeometricObjects and Engine module to GeoReferencing's PublicDependencyModuleNames. (Contributed by [REDACTED_PERSON])

Just some ws and alpha ordering fixup by GE.

diff --git a/Engine/Plugins/Runtime/GeoReferencing/Source/GeoReferencing/GeoReferencing.Build.cs b/Engine/Plugins/Runtime/GeoReferencing/Source/GeoReferencing/GeoReferencing.Build.cs
--- a/Engine/Plugins/Runtime/GeoReferencing/Source/GeoReferencing/GeoReferencing.Build.cs
+++ b/Engine/Plugins/Runtime/GeoReferencing/Source/GeoReferencing/GeoReferencing.Build.cs
@@ -15,7 +15,9 @@ public class GeoReferencing : ModuleRules
 		PublicDependencyModuleNames.AddRange(
 			new string[]
 			{
-				"Core"
+				"Core",
+				"Engine",
+				"GeometricObjects"
 			}
 		);
 
