GeometryCore: replace usage of GVector4 with TVector4
#rnx

diff --git a/Engine/Plugins/Runtime/GeoReferencing/Source/GeoReferencing/Private/GeoReferencingSystem.cpp b/Engine/Plugins/Runtime/GeoReferencing/Source/GeoReferencing/Private/GeoReferencingSystem.cpp
--- a/Engine/Plugins/Runtime/GeoReferencing/Source/GeoReferencing/Private/GeoReferencingSystem.cpp
+++ b/Engine/Plugins/Runtime/GeoReferencing/Source/GeoReferencing/Private/GeoReferencingSystem.cpp
@@ -112,10 +112,10 @@ void AGeoReferencingSystem::Initialize()
 	Impl->InitPROJLibrary();
 
 	Impl->WorldFrameToUEFrame = FMatrix4d(
-		UE::Geometry::GVector4(1.0, 0.0, 0.0, 0.0),
-		UE::Geometry::GVector4(0.0, -1.0, 0.0, 0.0),
-		UE::Geometry::GVector4(0.0, 0.0, 1.0, 0.0),
-		UE::Geometry::GVector4(0.0, 0.0, 0.0, 1.0), true);
+		FVector4d(1.0, 0.0, 0.0, 0.0),
+		FVector4d(0.0, -1.0, 0.0, 0.0),
+		FVector4d(0.0, 0.0, 1.0, 0.0),
+		FVector4d(0.0, 0.0, 0.0, 1.0), true);
 	Impl->UEFrameToWorldFrame = Impl->WorldFrameToUEFrame.Inverse();
 
 	ApplySettings();
@@ -836,8 +836,6 @@ PJ* AGeoReferencingSystem::FGeoReferencingSystemInternals::GetPROJProjection(FSt
 
 FMatrix4d AGeoReferencingSystem::FGeoReferencingSystemInternals::GetWorldFrameToECEFFrame(const FEllipsoid& Ellipsoid, const FCartesianCoordinates& ECEFLocation)
 {
-	using UE::Geometry::GVector4;
-
 	// See ECEF standard : https://commons.wikimedia.org/wiki/File:ECEF_ENU_Longitude_Latitude_right-hand-rule.svg
 	if (FMath::Abs(ECEFLocation.X) < FMathd::Epsilon &&
 		FMath::Abs(ECEFLocation.Y) < FMathd::Epsilon)
@@ -847,10 +845,10 @@ FMatrix4d AGeoReferencingSystem::FGeoReferencingSystemInternals::GetWorldFrameTo
 		{
 			// At origin - Should not happen, but consider it's the same as north pole
 			return FMatrix4d(
-				GVector4(0.0, 1.0, 0.0, 0.0),    // East = Y
-				GVector4(-1.0, 0.0, 0.0, 0.0), // North = -X
-				GVector4(0.0, 0.0, 1.0, 0.0),    // Up = Z
-				GVector4(ECEFLocation.X, ECEFLocation.Y, ECEFLocation.Z, 1.0), // Location
+				FVector4d(0.0, 1.0, 0.0, 0.0),    // East = Y
+				FVector4d(-1.0, 0.0, 0.0, 0.0), // North = -X
+				FVector4d(0.0, 0.0, 1.0, 0.0),    // Up = Z
+				FVector4d(ECEFLocation.X, ECEFLocation.Y, ECEFLocation.Z, 1.0), // Location
 				false); // Expressed in columns
 		}
 		else
@@ -858,10 +856,10 @@ FMatrix4d AGeoReferencingSystem::FGeoReferencingSystemInternals::GetWorldFrameTo
 			// At South or North pole - Axis are set to be continuous with other points
 			int Sign = FMathd::SignNonZero(ECEFLocation.Z);
 			return FMatrix4d(
-				GVector4(0.0, 1.0, 0.0, 0.0),    // East = Y
-				GVector4(-1.0 * Sign, 0.0, 0.0, 0.0), // North = -X
-				GVector4(0.0, 0.0, 1.0 * Sign, 0.0),    // Up = Z
-				GVector4(ECEFLocation.X, ECEFLocation.Y, ECEFLocation.Z, 1.0), // Location
+				FVector4d(0.0, 1.0, 0.0, 0.0),    // East = Y
+				FVector4d(-1.0 * Sign, 0.0, 0.0, 0.0), // North = -X
+				FVector4d(0.0, 0.0, 1.0 * Sign, 0.0),    // Up = Z
+				FVector4d(ECEFLocation.X, ECEFLocation.Y, ECEFLocation.Z, 1.0), // Location
 				false); // Expressed in columns
 		}
 	}
@@ -872,10 +870,10 @@ FMatrix4d AGeoReferencingSystem::FGeoReferencingSystemInternals::GetWorldFrameTo
 		FVector3d North = Up.Cross(East);
 
 		return FMatrix4d(
-			GVector4(East.X, East.Y, East.Z, 0.0),        // East
-			GVector4(North.X, North.Y, North.Z, 0.0),    // North
-			GVector4(Up.X, Up.Y, Up.Z, 0.0),            // Up
-			GVector4(ECEFLocation.X, ECEFLocation.Y, ECEFLocation.Z, 1.0), // Location
+			FVector4d(East.X, East.Y, East.Z, 0.0),        // East
+			FVector4d(North.X, North.Y, North.Z, 0.0),    // North
+			FVector4d(Up.X, Up.Y, Up.Z, 0.0),            // Up
+			FVector4d(ECEFLocation.X, ECEFLocation.Y, ECEFLocation.Z, 1.0), // Location
 			false); // Expressed in columns
 	}
 }
diff --git a/Engine/Plugins/Runtime/GeoReferencing/Source/GeoReferencing/Public/Matrix4d.h b/Engine/Plugins/Runtime/GeoReferencing/Source/GeoReferencing/Public/Matrix4d.h
--- a/Engine/Plugins/Runtime/GeoReferencing/Source/GeoReferencing/Public/Matrix4d.h
+++ b/Engine/Plugins/Runtime/GeoReferencing/Source/GeoReferencing/Public/Matrix4d.h
@@ -9,13 +9,15 @@ namespace UE
 namespace Geometry
 {
 
+using namespace UE::Math;
+
 template <typename RealType>
 struct TMatrix4
 {
-	GVector4<RealType> Row0;
-	GVector4<RealType> Row1;
-	GVector4<RealType> Row2;
-	GVector4<RealType> Row3;
+	TVector4<RealType> Row0;
+	TVector4<RealType> Row1;
+	TVector4<RealType> Row2;
+	TVector4<RealType> Row3;
 
 	TMatrix4()
 	{
@@ -23,7 +25,7 @@ struct TMatrix4
 
 	TMatrix4(RealType ConstantValue)
 	{
-		Row0 = GVector4<RealType>(ConstantValue, ConstantValue, ConstantValue, ConstantValue);
+		Row0 = TVector4<RealType>(ConstantValue, ConstantValue, ConstantValue, ConstantValue);
 		Row1 = Row0;
 		Row2 = Row0;
 		Row3 = Row0;
@@ -31,17 +33,17 @@ struct TMatrix4
 
 	TMatrix4(RealType Diag0, RealType Diag1, RealType Diag2, RealType Diag3)
 	{
-		Row0 = GVector4<RealType>(Diag0, 0, 0, 0);
-		Row1 = GVector4<RealType>(0, Diag1, 0, 0);
-		Row2 = GVector4<RealType>(0, 0, Diag2, 0);
-		Row3 = GVector4<RealType>(0, 0, 0, Diag3);
+		Row0 = TVector4<RealType>(Diag0, 0, 0, 0);
+		Row1 = TVector4<RealType>(0, Diag1, 0, 0);
+		Row2 = TVector4<RealType>(0, 0, Diag2, 0);
+		Row3 = TVector4<RealType>(0, 0, 0, Diag3);
 	}
 
 	/**
 	 * Construct outer-product of U*transpose(V) of U and V
 	 * result is that Mij = u_i * v_j
 	 */
-	TMatrix4<RealType>(const GVector4<RealType>& U, const GVector4<RealType>& V)
+	TMatrix4<RealType>(const TVector4<RealType>& U, const TVector4<RealType>& V)
 		: Row0(U.X* V.X, U.X* V.Y, U.X* V.Z, U.X* V.W),
 		Row1(U.Y* V.X, U.Y* V.Y, U.Y* V.Z, U.Y* V.W),
 		Row2(U.Z* V.X, U.Z* V.Y, U.Z* V.Z, U.Z* V.W),
@@ -57,7 +59,7 @@ struct TMatrix4
 	{
 	}
 
-	TMatrix4(const GVector4<RealType>& V1, const GVector4<RealType>& V2, const GVector4<RealType>& V3, const GVector4<RealType>& V4, bool bRows)
+	TMatrix4(const TVector4<RealType>& V1, const TVector4<RealType>& V2, const TVector4<RealType>& V3, const TVector4<RealType>& V4, bool bRows)
 	{
 		if (bRows)
 		{
@@ -68,10 +70,10 @@ struct TMatrix4
 		}
 		else
 		{
-			Row0 = GVector4<RealType>(V1.X, V2.X, V3.X, V4.X);
-			Row1 = GVector4<RealType>(V1.Y, V2.Y, V3.Y, V4.Y);
-			Row2 = GVector4<RealType>(V1.Z, V2.Z, V3.Z, V4.Z);
-			Row3 = GVector4<RealType>(V1.W, V2.W, V3.W, V4.W);
+			Row0 = TVector4<RealType>(V1.X, V2.X, V3.X, V4.X);
+			Row1 = TVector4<RealType>(V1.Y, V2.Y, V3.Y, V4.Y);
+			Row2 = TVector4<RealType>(V1.Z, V2.Z, V3.Z, V4.Z);
+			Row3 = TVector4<RealType>(V1.W, V2.W, V3.W, V4.W);
 		}
 	}
 
@@ -114,18 +116,18 @@ struct TMatrix4
 			Row3.X * Scale, Row3.Y * Scale, Row3.Z * Scale, Row3.W * Scale);
 	}
 
-	GVector4<RealType> operator*(const GVector4<RealType>& V) const
+	TVector4<RealType> operator*(const TVector4<RealType>& V) const
 	{
-		return GVector4<RealType>(
+		return TVector4<RealType>(
 			Row0.X * V.X + Row0.Y * V.Y + Row0.Z * V.Z + Row0.W * V.W,
 			Row1.X * V.X + Row1.Y * V.Y + Row1.Z * V.Z + Row1.W * V.W,
 			Row2.X * V.X + Row2.Y * V.Y + Row2.Z * V.Z + Row2.W * V.W,
 			Row3.X * V.X + Row3.Y * V.Y + Row3.Z * V.Z + Row3.W * V.W);
 	}
 
-	FVector3<RealType> operator*(const FVector3<RealType>& V) const
+	TVector<RealType> operator*(const TVector<RealType>& V) const
 	{
-		return FVector3<RealType>(
+		return TVector<RealType>(
 			Row0.X * V.X + Row0.Y * V.Y + Row0.Z * V.Z,
 			Row1.X * V.X + Row1.Y * V.Y + Row1.Z * V.Z,
 			Row2.X * V.X + Row2.Y * V.Y + Row2.Z * V.Z);
