Initial conversion  of GeoReferencing plugin for UE5

### Notes
GeoReferencing
- Fixed PROJ failure in PIE because of locked data file
- Replaced the previous math classes by the double ones from UE5
- Adapted the provided BP to new types
- Improved Probe/ProbeSpawners behavior
- Fixed the Tangent computations issue in Flat Terrain mode.

GeoReferencingDemo
- Rewamped the demo/test app and added more test cases

### Hashtags

diff --git a/Engine/Plugins/Runtime/GeoReferencing/GeoReferencing.uplugin b/Engine/Plugins/Runtime/GeoReferencing/GeoReferencing.uplugin
--- a/Engine/Plugins/Runtime/GeoReferencing/GeoReferencing.uplugin
+++ b/Engine/Plugins/Runtime/GeoReferencing/GeoReferencing.uplugin
@@ -3,7 +3,7 @@
 	"Version": 1,
 	"VersionName": "1.0",
 	"FriendlyName": "GeoReferencing",
-	"Description": "GeoReferencing tools for UE4 worlds",
+	"Description": "GeoReferencing tools for UE worlds",
 	"Category": "Other",
 	"CreatedBy": "Epic Games, Inc.",
 	"CreatedByURL": "",
@@ -43,10 +43,6 @@
 		}
 	],
 	"Plugins": [
-		{
-			"Name": "GeometryProcessing",
-			"Enabled": true
-		},
 		{
 			"Name": "SQLiteCore",
 			"Enabled": true
diff --git a/Engine/Plugins/Runtime/GeoReferencing/Source/GeoReferencing/GeoReferencing.Build.cs b/Engine/Plugins/Runtime/GeoReferencing/Source/GeoReferencing/GeoReferencing.Build.cs
--- a/Engine/Plugins/Runtime/GeoReferencing/Source/GeoReferencing/GeoReferencing.Build.cs
+++ b/Engine/Plugins/Runtime/GeoReferencing/Source/GeoReferencing/GeoReferencing.Build.cs
@@ -25,12 +25,12 @@ public class GeoReferencing : ModuleRules
 			{
 				"CoreUObject",
 				"Engine",
-				"GeometryCore",
 				"PROJ",
 				"Projects",
 				"Slate",
 				"SlateCore",
 				"SQLiteCore"
+
 			}
 		);
 
diff --git a/Engine/Plugins/Runtime/GeoReferencing/Source/GeoReferencing/Private/CartesianCoordinates.cpp b/Engine/Plugins/Runtime/GeoReferencing/Source/GeoReferencing/Private/CartesianCoordinates.cpp
--- a/Engine/Plugins/Runtime/GeoReferencing/Source/GeoReferencing/Private/CartesianCoordinates.cpp
+++ b/Engine/Plugins/Runtime/GeoReferencing/Source/GeoReferencing/Private/CartesianCoordinates.cpp
@@ -18,7 +18,7 @@ FCartesianCoordinates::FCartesianCoordinates(double InX, double InY, double InZ)
 
 }
 
-FCartesianCoordinates::FCartesianCoordinates(const FVector3d& Coordinates)
+FCartesianCoordinates::FCartesianCoordinates(const FVector& Coordinates)
 	: X(Coordinates.X)
 	, Y(Coordinates.Y)
 	, Z(Coordinates.Z)
@@ -78,7 +78,7 @@ void FCartesianCoordinates::ToFloatApproximation(float& OutX, float& OutY, float
 	OutZ = static_cast<float>(Z);
 }
 
-FVector3d FCartesianCoordinates::ToVector3d() const
+FVector FCartesianCoordinates::ToVector() const
 {
-	return FVector3d(X, Y, Z);
+	return FVector(X, Y, Z);
 }
diff --git a/Engine/Plugins/Runtime/GeoReferencing/Source/GeoReferencing/Private/Ellipsoid.cpp b/Engine/Plugins/Runtime/GeoReferencing/Source/GeoReferencing/Private/Ellipsoid.cpp
--- a/Engine/Plugins/Runtime/GeoReferencing/Source/GeoReferencing/Private/Ellipsoid.cpp
+++ b/Engine/Plugins/Runtime/GeoReferencing/Source/GeoReferencing/Private/Ellipsoid.cpp
@@ -1,6 +1,12 @@
 // Copyright Epic Games, Inc. All Rights Reserved.
 
 #include "Ellipsoid.h"
+#include "MathUtil.h"
+
+
+// TODOLWC - To be replaced once FVector::Normalize will use a smaller number than 1e-8
+#define GEOREF_DOUBLE_SMALL_NUMBER (1.e-50)
+
 
 FEllipsoid::FEllipsoid()
 	: FEllipsoid(1.0, 1.0, 1.0)
@@ -9,33 +15,35 @@ FEllipsoid::FEllipsoid()
 }
 
 FEllipsoid::FEllipsoid(double RadiusX, double RadiusY, double RadiusZ)
-	: FEllipsoid(FVector3d(RadiusX, RadiusY, RadiusZ))
+	: FEllipsoid(FVector(RadiusX, RadiusY, RadiusZ))
 {
 
 }
 
-FEllipsoid::FEllipsoid(const FVector3d& InRadii)
+FEllipsoid::FEllipsoid(const FVector& InRadii)
 	: Radii(InRadii)
 	, RadiiSquared(InRadii.X * InRadii.X, InRadii.Y * InRadii.Y, InRadii.Z * InRadii.Z)
 {
 	check(InRadii.X != 0 && InRadii.Y != 0 && InRadii.Z != 0);
 
-	OneOverRadii = FVector3d(1.0 / InRadii.X, 1.0 / InRadii.Y, 1.0 / InRadii.Z);
-	OneOverRadiiSquared = FVector3d(1.0 / (InRadii.X * InRadii.X), 1.0 / (InRadii.Y * InRadii.Y), 1.0 / (InRadii.Z * InRadii.Z));
+	OneOverRadii = FVector(1.0 / InRadii.X, 1.0 / InRadii.Y, 1.0 / InRadii.Z);
+	OneOverRadiiSquared = FVector(1.0 / (InRadii.X * InRadii.X), 1.0 / (InRadii.Y * InRadii.Y), 1.0 / (InRadii.Z * InRadii.Z));
 }
 
-FVector3d FEllipsoid::GeodeticSurfaceNormal(const FCartesianCoordinates& ECEFLocation) const
+FVector FEllipsoid::GeodeticSurfaceNormal(const FCartesianCoordinates& ECEFLocation) const
 {
-	FVector3d Normal( ECEFLocation.X * OneOverRadiiSquared.X, ECEFLocation.Y * OneOverRadiiSquared.Y, ECEFLocation.Z * OneOverRadiiSquared.Z);
-
-	return UE::Geometry::Normalized(Normal);
+	FVector Normal( ECEFLocation.X * OneOverRadiiSquared.X, ECEFLocation.Y * OneOverRadiiSquared.Y, ECEFLocation.Z * OneOverRadiiSquared.Z);
+	Normal.Normalize(GEOREF_DOUBLE_SMALL_NUMBER); 
+	return Normal;
 }
 
-FVector3d FEllipsoid::GeodeticSurfaceNormal(const FGeographicCoordinates& GeographicCoordinates) const
+FVector FEllipsoid::GeodeticSurfaceNormal(const FGeographicCoordinates& GeographicCoordinates) const
 {
 	double LongitudeRad = FMathd::DegToRad * GeographicCoordinates.Longitude ;
 	double LatitudeRad = FMathd::DegToRad * GeographicCoordinates.Latitude;
 	double cosLatitude = FMathd::Cos(LatitudeRad);
 
-	return UE::Geometry::Normalized(FVector3d(cosLatitude * FMathd::Cos(LongitudeRad), cosLatitude * FMathd::Sin(LongitudeRad), FMathd::Sin(LatitudeRad)));
+	FVector Normal(cosLatitude * FMathd::Cos(LongitudeRad), cosLatitude * FMathd::Sin(LongitudeRad), FMathd::Sin(LatitudeRad));
+	Normal.Normalize(GEOREF_DOUBLE_SMALL_NUMBER);
+	return Normal;
 }
diff --git a/Engine/Plugins/Runtime/GeoReferencing/Source/GeoReferencing/Private/GeoReferencingSystem.cpp b/Engine/Plugins/Runtime/GeoReferencing/Source/GeoReferencing/Private/GeoReferencingSystem.cpp
--- a/Engine/Plugins/Runtime/GeoReferencing/Source/GeoReferencing/Private/GeoReferencingSystem.cpp
+++ b/Engine/Plugins/Runtime/GeoReferencing/Source/GeoReferencing/Private/GeoReferencingSystem.cpp
@@ -9,6 +9,7 @@
 #include "Framework/Notifications/NotificationManager.h"
 #include "Engine/Engine.h"
 #include "Engine/World.h"
+#include "MathUtil.h"
 
 #include "HAL/PlatformFileManager.h"
 #include "IPlatformFilePak.h"
@@ -22,6 +23,10 @@ THIRD_PARTY_INCLUDES_END
 
 #define ECEF_EPSG_FSTRING FString(TEXT("EPSG:4978"))
 
+// TODOLWC - To be replaced once FVector::Normalize will use a smaller number than 1e-8
+#define GEOREF_DOUBLE_SMALL_NUMBER			(1.e-50)
+
+
 DEFINE_LOG_CATEGORY(LogGeoReferencing);
 
 AGeoReferencingSystem* AGeoReferencingSystem::GetGeoReferencingSystem(UObject* WorldContextObject)
@@ -66,9 +71,7 @@ public:
 	void DeInitPROJLibrary();
 	PJ* GetPROJProjection(FString SourceCRS, FString DestinationCRS);
 	bool GetEllipsoid(FString CRSString, FEllipsoid& Ellipsoid);
-	FMatrix4d GetWorldFrameToECEFFrame(const FEllipsoid& Ellipsoid, const FCartesianCoordinates& ECEFLocation);
-
-	static FMatrix ConvertMatrix4d(FMatrix4d Matrix);
+	FMatrix GetWorldFrameToECEFFrame(const FEllipsoid& Ellipsoid, const FCartesianCoordinates& ECEFLocation);
 
 	PJ_CONTEXT* ProjContext;
 	PJ* ProjProjectedToGeographic;
@@ -79,14 +82,14 @@ public:
 
 	// Transformation caches 
 	// Flat Planet
-	FVector3d WorldOriginLocationProjected; // Offset between the UE world and the Projected CRS Origin. (Expressed in ProjectedCRS units).
+	FVector WorldOriginLocationProjected; // Offset between the UE world and the Projected CRS Origin. (Expressed in ProjectedCRS units).
 
 	// Round Planet
-	FMatrix4d WorldFrameToECEFFrame; // Matrix to transform a vector from EU to ECEF CRS
-	FMatrix4d ECEFFrameToWorldFrame; // Matrix to transform a vector from ECEF to UE CRS - Inverse of the previous one kept in cache. 
+	FMatrix WorldFrameToECEFFrame; // Matrix to transform a vector from EU to ECEF CRS
+	FMatrix ECEFFrameToWorldFrame; // Matrix to transform a vector from ECEF to UE CRS - Inverse of the previous one kept in cache. 
 
-	FMatrix4d WorldFrameToUEFrame;
-	FMatrix4d UEFrameToWorldFrame;
+	FMatrix WorldFrameToUEFrame;
+	FMatrix UEFrameToWorldFrame;
 };
 
 /////// INIT / DEINIT
@@ -111,11 +114,14 @@ void AGeoReferencingSystem::Initialize()
 
 	Impl->InitPROJLibrary();
 
-	Impl->WorldFrameToUEFrame = FMatrix4d(
-		FVector4d(1.0, 0.0, 0.0, 0.0),
-		FVector4d(0.0, -1.0, 0.0, 0.0),
-		FVector4d(0.0, 0.0, 1.0, 0.0),
-		FVector4d(0.0, 0.0, 0.0, 1.0), true);
+	// TODOABLWC - Can we consider other conventions ? Or North Offset like in SunPosition?
+	Impl->WorldFrameToUEFrame = FMatrix( 
+		FVector(1.0, 0.0, 0.0),		// Easting (X) is UE World X
+		FVector(0.0, -1.0, 0.0),	// Northing (Y) is UE World -Y because of left-handed convention
+		FVector(0.0, 0.0, 1.0),		// Up (Z) is UE World Z 
+		FVector(0.0, 0.0, 0.0));	// No Origin offset
+
+
 	Impl->UEFrameToWorldFrame = Impl->WorldFrameToUEFrame.Inverse();
 
 	ApplySettings();
@@ -150,16 +156,16 @@ void AGeoReferencingSystem::EngineToProjected(const FVector& EngineCoordinates,
 		// in FlatPlanet, the transform is simply a translation
 		// Before any conversion, consider the internal UE rebasing
 		FIntVector UERebasingOffset = GetWorld()->OriginLocation;
-		FVector3d UERebasedCoordinates(
+		FVector UERebasedCoordinates(
 			EngineCoordinates.X + UERebasingOffset.X,
 			EngineCoordinates.Y + UERebasingOffset.Y,
 			EngineCoordinates.Z + UERebasingOffset.Z);
 
 		// Convert UE units to meters, invert the Y coordinate because of left-handed UE Frame
-		FVector3d UEWorldCoordinates = UERebasedCoordinates * FVector3d(0.01, -0.01, 0.01);
+		FVector UEWorldCoordinates = UERebasedCoordinates * FVector(0.01, -0.01, 0.01);
 
 		// Add the defined origin offset
-		FVector3d Projected = UEWorldCoordinates + Impl->WorldOriginLocationProjected;
+		FVector Projected = UEWorldCoordinates + Impl->WorldOriginLocationProjected;
 		ProjectedCoordinates = FCartesianCoordinates(Projected);
 	}
 	break;
@@ -184,15 +190,15 @@ void AGeoReferencingSystem::ProjectedToEngine(const FCartesianCoordinates& Proje
 	{
 		// in FlatPlanet, the transform is simply a translation
 		// Remove the Origin location, and convert to UE Units, while inverting the Z Axis
-		FVector3d Projected(ProjectedCoordinates.X, ProjectedCoordinates.Y, ProjectedCoordinates.Z);
-		FVector3d UEWorldCoordinates = (Projected - Impl->WorldOriginLocationProjected);
+		FVector Projected(ProjectedCoordinates.X, ProjectedCoordinates.Y, ProjectedCoordinates.Z);
+		FVector UEWorldCoordinates = (Projected - Impl->WorldOriginLocationProjected);
 
 		// Convert UE units to meters, invert the Y coordinate because of left-handed UE Frame
-		FVector3d UERebasedCoordinates = UEWorldCoordinates * FVector3d(100.0, -100.0, 100.0);
+		FVector UERebasedCoordinates = UEWorldCoordinates * FVector(100.0, -100.0, 100.0);
 
 		// Consider the UE internal rebasing
 		FIntVector UERebasingOffset = GetWorld()->OriginLocation;
-		FVector3d EngineCoordinates3d = UERebasedCoordinates - FVector3d(UERebasingOffset.X, UERebasingOffset.Y, UERebasingOffset.Z);
+		FVector EngineCoordinates3d = UERebasedCoordinates - FVector(UERebasingOffset.X, UERebasingOffset.Y, UERebasingOffset.Z);
 
 		EngineCoordinates = FVector(EngineCoordinates3d.X, EngineCoordinates3d.Y, EngineCoordinates3d.Z);
 	}
@@ -208,13 +214,13 @@ void AGeoReferencingSystem::EngineToECEF(const FVector& EngineCoordinates, FCart
 	{
 		// Before any conversion, consider the internal UE rebasing
 		FIntVector UERebasingOffset = GetWorld()->OriginLocation;
-		FVector3d UERebasedCoordinates(
+		FVector UERebasedCoordinates(
 			EngineCoordinates.X + UERebasingOffset.X,
 			EngineCoordinates.Y + UERebasingOffset.Y,
 			EngineCoordinates.Z + UERebasingOffset.Z);
 
 		// Convert UE units to meters, invert the Y coordinate because of left-handed UE Frame
-		FVector3d UEWorldCoordinates = UERebasedCoordinates * FVector3d(0.01, -0.01, 0.01);
+		FVector UEWorldCoordinates = UERebasedCoordinates * FVector(0.01, -0.01, 0.01);
 
 		if (bOriginAtPlanetCenter)
 		{
@@ -223,7 +229,7 @@ void AGeoReferencingSystem::EngineToECEF(const FVector& EngineCoordinates, FCart
 		}
 		else
 		{
-			FVector4d ECEFCoordinates4d = Impl->WorldFrameToECEFFrame * FVector4d(UEWorldCoordinates.X, UEWorldCoordinates.Y, UEWorldCoordinates.Z, 1.0);
+			FVector ECEFCoordinates4d = Impl->WorldFrameToECEFFrame.TransformPosition( FVector(UEWorldCoordinates.X, UEWorldCoordinates.Y, UEWorldCoordinates.Z));
 			ECEFCoordinates = FCartesianCoordinates(ECEFCoordinates4d);
 		}
 	}
@@ -247,26 +253,23 @@ void AGeoReferencingSystem::ECEFToEngine(const FCartesianCoordinates& ECEFCoordi
 	{
 	case EPlanetShape::RoundPlanet:
 	{
-		FVector3d UEWorldCoordinates;
+		FVector UEWorldCoordinates;
 		if (bOriginAtPlanetCenter)
 		{
 			// Easy case, UE is ECEF... And we did the rebasing, so return the Global coordinates
-			UEWorldCoordinates = FVector3d(ECEFCoordinates.X, ECEFCoordinates.Y, ECEFCoordinates.Z);
+			UEWorldCoordinates = FVector(ECEFCoordinates.X, ECEFCoordinates.Y, ECEFCoordinates.Z);
 		}
 		else
 		{
-			FVector4d UEWorldCoordinates4d = Impl->ECEFFrameToWorldFrame * FVector4d(ECEFCoordinates.X, ECEFCoordinates.Y, ECEFCoordinates.Z, 1.0);
-			UEWorldCoordinates.X = UEWorldCoordinates4d.X;
-			UEWorldCoordinates.Y = UEWorldCoordinates4d.Y;
-			UEWorldCoordinates.Z = UEWorldCoordinates4d.Z;
+			UEWorldCoordinates = Impl->ECEFFrameToWorldFrame.TransformPosition(FVector(ECEFCoordinates.X, ECEFCoordinates.Y, ECEFCoordinates.Z));
 		}
 
 		// Convert UE units to meters, invert the Y coordinate because of left-handed UE Frame
-		FVector3d UERebasedCoordinates = UEWorldCoordinates * FVector3d(100.0, -100.0, 100.0);
+		FVector UERebasedCoordinates = UEWorldCoordinates * FVector(100.0, -100.0, 100.0);
 
 		// Consider the UE internal rebasing
 		FIntVector UERebasingOffset = GetWorld()->OriginLocation;
-		FVector3d EngineCoordinates3d = UERebasedCoordinates - FVector3d(UERebasingOffset.X, UERebasingOffset.Y, UERebasingOffset.Z);
+		FVector EngineCoordinates3d = UERebasedCoordinates - FVector(UERebasingOffset.X, UERebasingOffset.Y, UERebasingOffset.Z);
 
 		EngineCoordinates = FVector(EngineCoordinates3d.X, EngineCoordinates3d.Y, EngineCoordinates3d.Z);
 	}
@@ -382,40 +385,37 @@ void AGeoReferencingSystem::GetENUVectorsAtECEFLocation(const FCartesianCoordina
 	{
 		Ellipsoid = Impl->ProjectedEllipsoid;
 	}
-	FMatrix4d WorldFrameToECEFFrameAtLocation;
+	FMatrix WorldFrameToECEFFrameAtLocation;
 	WorldFrameToECEFFrameAtLocation = Impl->GetWorldFrameToECEFFrame(Ellipsoid, ECEFCoordinates);
 
 	switch (PlanetShape)
 	{
 	case EPlanetShape::RoundPlanet:
 	{
-		FMatrix4d UEtoECEF = Impl->UEFrameToWorldFrame * Impl->ECEFFrameToWorldFrame * WorldFrameToECEFFrameAtLocation; //TOOD?
-
-		East = FVector(UEtoECEF.Row0.X, UEtoECEF.Row1.X, UEtoECEF.Row2.X);
-		North = FVector(UEtoECEF.Row0.Y, UEtoECEF.Row1.Y, UEtoECEF.Row2.Y);
-		Up = FVector(UEtoECEF.Row0.Z, UEtoECEF.Row1.Z, UEtoECEF.Row2.Z);
+		FMatrix UEtoECEF = WorldFrameToECEFFrameAtLocation * Impl->ECEFFrameToWorldFrame * Impl->UEFrameToWorldFrame;
+		UEtoECEF.GetUnitAxes(East, North, Up);
 	}
 	break;
 
 	case EPlanetShape::FlatPlanet:
 	default:
 		// PROJ don't provide anything to project direction vectors. Let's do it by hand...
-		FVector3d EasternPoint = ECEFCoordinates.ToVector3d() + WorldFrameToECEFFrameAtLocation * FVector3d(1.0, 0.0, 0.0);
-		FVector3d NorthernPoint = ECEFCoordinates.ToVector3d() + WorldFrameToECEFFrameAtLocation * FVector3d(0.0, 1.0, 0.0);
-
+		FVector EasternPoint = ECEFCoordinates.ToVector() + WorldFrameToECEFFrameAtLocation.TransformVector(FVector(1.0, 0.0, 0.0)); // 1m from origin to the East
+		FVector NorthernPoint = ECEFCoordinates.ToVector() + WorldFrameToECEFFrameAtLocation.TransformVector(FVector(0.0, 1.0, 0.0)); // 1m from origin to the North
+		
 		FCartesianCoordinates ProjectedOrigin, ProjectedEastern, ProjectedNorthern;
 		ECEFToProjected(ECEFCoordinates, ProjectedOrigin);
 		ECEFToProjected(EasternPoint, ProjectedEastern);
 		ECEFToProjected(NorthernPoint, ProjectedNorthern);
 
-		FVector3d East3d(ProjectedEastern.X - ProjectedOrigin.X, ProjectedEastern.Y - ProjectedOrigin.Y, ProjectedEastern.Z - ProjectedOrigin.Z);
-		FVector3d North3d(ProjectedNorthern.X - ProjectedOrigin.X, ProjectedNorthern.Y - ProjectedOrigin.Y, ProjectedNorthern.Z - ProjectedOrigin.Z);
+		FVector EastDirection(ProjectedEastern.X - ProjectedOrigin.X, ProjectedEastern.Y - ProjectedOrigin.Y, ProjectedEastern.Z - ProjectedOrigin.Z);
+		FVector NorthDirection(ProjectedNorthern.X - ProjectedOrigin.X, ProjectedNorthern.Y - ProjectedOrigin.Y, ProjectedNorthern.Z - ProjectedOrigin.Z);
 
-		East3d.Normalize();
-		North3d.Normalize();
+		EastDirection.Normalize(GEOREF_DOUBLE_SMALL_NUMBER);
+		NorthDirection.Normalize(GEOREF_DOUBLE_SMALL_NUMBER);
 
-		East = FVector(East3d.X, -East3d.Y, East3d.Z);
-		North = FVector(North3d.X, -North3d.Y, North3d.Z);
+		East = FVector(EastDirection.X, -EastDirection.Y, EastDirection.Z);
+		North = FVector(NorthDirection.X, -NorthDirection.Y, NorthDirection.Z);
 		Up = FVector::CrossProduct(North, East);
 		break;
 	}
@@ -431,12 +431,9 @@ void AGeoReferencingSystem::GetECEFENUVectorsAtECEFLocation(const FCartesianCoor
 		Ellipsoid = Impl->ProjectedEllipsoid;
 	}
 	
-	FMatrix4d WorldFrameToECEFFrameAtLocation;
+	FMatrix WorldFrameToECEFFrameAtLocation;
 	WorldFrameToECEFFrameAtLocation = Impl->GetWorldFrameToECEFFrame(Ellipsoid, ECEFCoordinates);
-
-	ECEFEast = FVector(WorldFrameToECEFFrameAtLocation.Row0.X, WorldFrameToECEFFrameAtLocation.Row1.X, WorldFrameToECEFFrameAtLocation.Row2.X);
-	ECEFNorth = FVector(WorldFrameToECEFFrameAtLocation.Row0.Y, WorldFrameToECEFFrameAtLocation.Row1.Y, WorldFrameToECEFFrameAtLocation.Row2.Y);
-	ECEFUp = FVector(WorldFrameToECEFFrameAtLocation.Row0.Z, WorldFrameToECEFFrameAtLocation.Row1.Z, WorldFrameToECEFFrameAtLocation.Row2.Z);
+	WorldFrameToECEFFrameAtLocation.GetUnitAxes(ECEFEast, ECEFNorth, ECEFUp);
 }
 
 FTransform AGeoReferencingSystem::GetTangentTransformAtEngineLocation(const FVector& EngineCoordinates)
@@ -470,24 +467,24 @@ FTransform AGeoReferencingSystem::GetTangentTransformAtECEFLocation(const FCarte
 		{
 			Ellipsoid = Impl->ProjectedEllipsoid;
 		}
-		FMatrix4d WorldFrameToECEFFrameAtLocation;
-		WorldFrameToECEFFrameAtLocation = Impl->GetWorldFrameToECEFFrame(Ellipsoid, ECEFCoordinates);
-
+		FMatrix WorldFrameToECEFFrameAtLocation = Impl->GetWorldFrameToECEFFrame(Ellipsoid, ECEFCoordinates);
+		
 
-		FMatrix4d UEtoECEF = Impl->UEFrameToWorldFrame * Impl->ECEFFrameToWorldFrame * WorldFrameToECEFFrameAtLocation * Impl->UEFrameToWorldFrame; //TOOD?
+		FMatrix UEtoECEF = Impl->UEFrameToWorldFrame * WorldFrameToECEFFrameAtLocation * Impl->ECEFFrameToWorldFrame * Impl->WorldFrameToUEFrame; //TOOD?
+		FVector UEOrigin;
+		ECEFToEngine(ECEFCoordinates, UEOrigin);
+		UEtoECEF.SetOrigin(UEOrigin);
 
-		FVector East = FVector(UEtoECEF.Row0.X, UEtoECEF.Row1.X, UEtoECEF.Row2.X);
-		FVector North = FVector(UEtoECEF.Row0.Y, UEtoECEF.Row1.Y, UEtoECEF.Row2.Y);
-		FVector Up = FVector(UEtoECEF.Row0.Z, UEtoECEF.Row1.Z, UEtoECEF.Row2.Z);
-		FVector EngineLocation;
-		ECEFToEngine(ECEFCoordinates, EngineLocation);
-
-		FMatrix ResultMatrix = FMatrix::Identity;
-		ResultMatrix.SetAxes(&East, &North, &Up, &EngineLocation);
-		return FTransform(ResultMatrix);
+		return FTransform(UEtoECEF);
+	}
+	else
+	{
+		FVector East, North, Up, UEOrigin;
+		GetENUVectorsAtECEFLocation(ECEFCoordinates, East, North, Up);
+		ECEFToEngine(ECEFCoordinates, UEOrigin);
+		FMatrix TransformMatrix(East, North, Up, UEOrigin);
+		return FTransform(TransformMatrix);
 	}
-
-	return FTransform::Identity;
 }
 
 FTransform AGeoReferencingSystem::GetPlanetCenterTransform()
@@ -501,40 +498,26 @@ FTransform AGeoReferencingSystem::GetPlanetCenterTransform()
 		}
 		else
 		{
-			FMatrix4d TransformMatrix = Impl->UEFrameToWorldFrame * Impl->ECEFFrameToWorldFrame * Impl->UEFrameToWorldFrame;
+			FMatrix TransformMatrix = Impl->UEFrameToWorldFrame * Impl->ECEFFrameToWorldFrame * Impl->WorldFrameToUEFrame;
 			// Don't go to transform yet, we must stay in double to apply the rebasing offset. 
 
-			// Get Origin
-			FVector3d Origin3d(TransformMatrix.Row0.W, TransformMatrix.Row1.W, TransformMatrix.Row2.W);
-			// Convert UE units to meters, invert the Y coordinate because of left-handed UE Frame
-			Origin3d = Origin3d * FVector3d(100.0, -100.0, 100.0);
-
-
-			// Consider the UE internal rebasing
+			// Get Origin, and convert UE units to meters, invert the Y coordinate because of left-handed UE Frame
+			FVector UEOrigin = TransformMatrix.GetOrigin() * FVector(100.0, -100.0, 100.0);
+			
+			// Consider the UE internal rebasing to compute Origin
 			FIntVector UERebasingOffset = GetWorld()->OriginLocation;
-			FVector3d UERebasedCoordinates = Origin3d - FVector3d(UERebasingOffset.X, UERebasingOffset.Y, UERebasingOffset.Z);
+			FVector UERebasedCoordinates = UEOrigin - FVector(UERebasingOffset.X, UERebasingOffset.Y, UERebasingOffset.Z);
 			FVector Origin = FVector(UERebasedCoordinates.X, UERebasedCoordinates.Y, UERebasedCoordinates.Z);
+			TransformMatrix.SetOrigin(Origin);
 
-
-			FVector East, North, Up;
-			East = FVector(TransformMatrix.Row0.X, TransformMatrix.Row1.X, TransformMatrix.Row2.X);
-			North = FVector(TransformMatrix.Row0.Y, TransformMatrix.Row1.Y, TransformMatrix.Row2.Y);
-			Up = FVector(TransformMatrix.Row0.Z, TransformMatrix.Row1.Z, TransformMatrix.Row2.Z);
-
-
-			FMatrix ResultMatrix = FMatrix::Identity;
-			ResultMatrix.SetAxes(&East, &North, &Up, &Origin);
-			return FTransform(ResultMatrix);
+			return FTransform(TransformMatrix);
 		}
-
 	}
 	else
 	{
 		// Makes not sense in Flat planet mode... 
 		return FTransform::Identity;
 	}
-
-
 }
 
 // Public PROJ Utilities
@@ -668,15 +651,15 @@ void AGeoReferencingSystem::ApplySettings()
 		if (bOriginAtPlanetCenter)
 		{
 			// Theoritically, we should never use this Identity matrices since the transformations already handle that using a shorter code path, but let's keep consistency  
-			Impl->ECEFFrameToWorldFrame = FMatrix4d::Identity();
-			Impl->WorldFrameToECEFFrame = FMatrix4d::Identity();
+			Impl->ECEFFrameToWorldFrame.SetIdentity();
+			Impl->WorldFrameToECEFFrame.SetIdentity();
 		}
 		else
 		{
 			// We need to compute the ENU Vectors at the origin point. For that, the origin has to be expressed first in ECEF. 
 			// Express origin in ECEF, and get the ENU vectors
 			FCartesianCoordinates ECEFOrigin;
-			Impl->WorldFrameToECEFFrame = FMatrix4d::Identity();
+			Impl->WorldFrameToECEFFrame.SetIdentity();
 			if (bOriginLocationInProjectedCRS)
 			{
 				FCartesianCoordinates ProjectedOrigin(
@@ -701,7 +684,7 @@ void AGeoReferencingSystem::ApplySettings()
 		if (bOriginLocationInProjectedCRS)
 		{
 			// World origin is expressed using Projected coordinates. Take them as is (in double)
-			Impl->WorldOriginLocationProjected = FVector3d(
+			Impl->WorldOriginLocationProjected = FVector(
 				OriginProjectedCoordinatesEasting,
 				OriginProjectedCoordinatesNorthing,
 				OriginProjectedCoordinatesUp);
@@ -711,7 +694,7 @@ void AGeoReferencingSystem::ApplySettings()
 			// World origin is expressed using Geographic coordinates. Convert them to the projected CRS to have the offset
 			FCartesianCoordinates OriginProjected;
 			GeographicToProjected(FGeographicCoordinates(OriginLongitude, OriginLatitude, OriginAltitude), OriginProjected);
-			Impl->WorldOriginLocationProjected = FVector3d(OriginProjected.X, OriginProjected.Y, OriginProjected.Z);
+			Impl->WorldOriginLocationProjected = FVector(OriginProjected.X, OriginProjected.Y, OriginProjected.Z);
 		}
 		break;
 	}
@@ -759,6 +742,7 @@ void AGeoReferencingSystem::FGeoReferencingSystemInternals::InitPROJLibrary()
 	};
 	proj_context_set_search_paths(ProjContext, sizeof(ProjSearchPaths)/sizeof(ProjSearchPaths[0]), ProjSearchPaths);
 
+	proj_context_set_autoclose_database(ProjContext, true);
 	// Non-editor builds use UFS extensions to read PROJ data from UFS/Pak
 	if (!GIsEditor)
 	{
@@ -834,60 +818,41 @@ PJ* AGeoReferencingSystem::FGeoReferencingSystemInternals::GetPROJProjection(FSt
 	return P_for_GIS;
 }
 
-FMatrix4d AGeoReferencingSystem::FGeoReferencingSystemInternals::GetWorldFrameToECEFFrame(const FEllipsoid& Ellipsoid, const FCartesianCoordinates& ECEFLocation)
+FMatrix AGeoReferencingSystem::FGeoReferencingSystemInternals::GetWorldFrameToECEFFrame(const FEllipsoid& Ellipsoid, const FCartesianCoordinates& ECEFLocation)
 {
 	// See ECEF standard : https://commons.wikimedia.org/wiki/File:ECEF_ENU_Longitude_Latitude_right-hand-rule.svg
-	if (FMath::Abs(ECEFLocation.X) < FMathd::Epsilon &&
-		FMath::Abs(ECEFLocation.Y) < FMathd::Epsilon)
+	if (FMathd::Abs(ECEFLocation.X) < FMathd::Epsilon &&
+		FMathd::Abs(ECEFLocation.Y) < FMathd::Epsilon)
 	{
 		// Special Case - On earth axis... 
-		if (FMath::Abs(ECEFLocation.Z) < FMathd::Epsilon)
+		int Sign = 1;
+		if (FMathd::Abs(ECEFLocation.Z) < FMathd::Epsilon)
 		{
 			// At origin - Should not happen, but consider it's the same as north pole
-			return FMatrix4d(
-				FVector4d(0.0, 1.0, 0.0, 0.0),    // East = Y
-				FVector4d(-1.0, 0.0, 0.0, 0.0), // North = -X
-				FVector4d(0.0, 0.0, 1.0, 0.0),    // Up = Z
-				FVector4d(ECEFLocation.X, ECEFLocation.Y, ECEFLocation.Z, 1.0), // Location
-				false); // Expressed in columns
+			// Leave Sign = 1
 		}
 		else
 		{
 			// At South or North pole - Axis are set to be continuous with other points
-			int Sign = FMathd::SignNonZero(ECEFLocation.Z);
-			return FMatrix4d(
-				FVector4d(0.0, 1.0, 0.0, 0.0),    // East = Y
-				FVector4d(-1.0 * Sign, 0.0, 0.0, 0.0), // North = -X
-				FVector4d(0.0, 0.0, 1.0 * Sign, 0.0),    // Up = Z
-				FVector4d(ECEFLocation.X, ECEFLocation.Y, ECEFLocation.Z, 1.0), // Location
-				false); // Expressed in columns
+			Sign = FMathd::SignNonZero(ECEFLocation.Z);
 		}
+
+		return FMatrix(
+			FVector::YAxisVector, 			// East = Y
+			-FVector::XAxisVector * Sign,	// North = Sign * X
+			FVector::ZAxisVector * Sign,	// Up = Sign*Z
+			ECEFLocation.ToVector());
 	}
 	else
 	{
-		FVector3d Up = Ellipsoid.GeodeticSurfaceNormal(ECEFLocation);
-		FVector3d East = UE::Geometry::Normalized(FVector3d(-ECEFLocation.Y, ECEFLocation.X, 0.0));
-		FVector3d North = Up.Cross(East);
-
-		return FMatrix4d(
-			FVector4d(East.X, East.Y, East.Z, 0.0),        // East
-			FVector4d(North.X, North.Y, North.Z, 0.0),    // North
-			FVector4d(Up.X, Up.Y, Up.Z, 0.0),            // Up
-			FVector4d(ECEFLocation.X, ECEFLocation.Y, ECEFLocation.Z, 1.0), // Location
-			false); // Expressed in columns
+		FVector Up = Ellipsoid.GeodeticSurfaceNormal(ECEFLocation); //TODOABLWC - Mauvaises valeurs ici... 
+		FVector East(-ECEFLocation.Y, ECEFLocation.X, 0.0); 
+		East.Normalize(GEOREF_DOUBLE_SMALL_NUMBER);
+		FVector North = Up.Cross(East);
+		return FMatrix(	East, North, Up, ECEFLocation.ToVector());
 	}
 }
 
-FMatrix AGeoReferencingSystem::FGeoReferencingSystemInternals::ConvertMatrix4d(FMatrix4d Matrix)
-{
-	return FMatrix(
-		FVector(Matrix.Row0.X, Matrix.Row1.X, Matrix.Row2.X), // 1st Column
-		FVector(Matrix.Row0.Y, Matrix.Row1.Y, Matrix.Row2.Y), // 2nd Column
-		FVector(Matrix.Row0.Z, Matrix.Row1.Z, Matrix.Row2.Z), // 3rd Column
-		FVector(Matrix.Row0.W, Matrix.Row1.W, Matrix.Row2.W)  // 4th Column - Translation
-	);
-}
-
 #if WITH_EDITOR
 void AGeoReferencingSystem::PostEditChangeProperty(FPropertyChangedEvent& PropertyChangedEvent)
 {
diff --git a/Engine/Plugins/Runtime/GeoReferencing/Source/GeoReferencing/Public/CartesianCoordinates.h b/Engine/Plugins/Runtime/GeoReferencing/Source/GeoReferencing/Public/CartesianCoordinates.h
--- a/Engine/Plugins/Runtime/GeoReferencing/Source/GeoReferencing/Public/CartesianCoordinates.h
+++ b/Engine/Plugins/Runtime/GeoReferencing/Source/GeoReferencing/Public/CartesianCoordinates.h
@@ -4,16 +4,11 @@
 
 #include "CoreMinimal.h"
 #include "Kismet/BlueprintFunctionLibrary.h"
-
-// Using double precision maths from the Geometry plugin. 
-#include "VectorTypes.h"
-
 #include "CartesianCoordinates.generated.h"
 
-/// PEER_REVIEW - I used an Fxxx instead of UObject-based as I thought it was cheaper memorywise, 
-/// but I needed to add a BP_FL to make the conversions since Fstruct can't contain UFUNCTIONS. Is that a good choice ?
 
 USTRUCT(BlueprintType)
+
 struct GEOREFERENCING_API FCartesianCoordinates
 {
 	GENERATED_USTRUCT_BODY()
@@ -21,8 +16,11 @@ struct GEOREFERENCING_API FCartesianCoordinates
 public:
 	FCartesianCoordinates();
 	FCartesianCoordinates(double InX, double InY, double InZ);
-	FCartesianCoordinates(const FVector3d& Coordinates);
-	FCartesianCoordinates(const FVector4d& Coordinates);
+	FCartesianCoordinates(const FVector& Coordinates);
+
+	double X;
+	double Y;
+	double Z;
 
 	FText ToFullText(int32 IntegralDigits = 3);
 	FText ToCompactText(int32 IntegralDigits = 3);
@@ -30,11 +28,19 @@ public:
 
 	void ToFloatApproximation(float& OutX, float& OutY, float& OutZ);
 
-	FVector3d ToVector3d() const;
+	FVector ToVector() const;
+	
+	PRAGMA_DISABLE_DEPRECATION_WARNINGS // Suppress compiler warning on override of deprecated function
+
+
+		UE_DEPRECATED(5.0, "Use FCartesianCoordinates::ToVector() instead.")
+		FVector3d ToVector3d() const;
+		UE_DEPRECATED(5.0, "Use FCartesianCoordinates(const FVector& Coordinates) instead.")
+		FCartesianCoordinates(const FVector4d& Coordinates);
+
+	PRAGMA_ENABLE_DEPRECATION_WARNINGS
+
 	
-	double X;
-	double Y;
-	double Z;
 };
 
 UCLASS()
diff --git a/Engine/Plugins/Runtime/GeoReferencing/Source/GeoReferencing/Public/Ellipsoid.h b/Engine/Plugins/Runtime/GeoReferencing/Source/GeoReferencing/Public/Ellipsoid.h
--- a/Engine/Plugins/Runtime/GeoReferencing/Source/GeoReferencing/Public/Ellipsoid.h
+++ b/Engine/Plugins/Runtime/GeoReferencing/Source/GeoReferencing/Public/Ellipsoid.h
@@ -4,8 +4,6 @@
 
 #include "CoreMinimal.h"
 
-// Using double precision maths from the Geometry plugin. 
-#include "VectorTypes.h"
 #include "GeographicCoordinates.h"
 #include "CartesianCoordinates.h"
 
@@ -19,13 +17,13 @@ struct GEOREFERENCING_API FEllipsoid
 public:
 	FEllipsoid();
 	FEllipsoid(double RadiusX, double RadiusY, double RadiusZ);
-	FEllipsoid(const FVector3d& InRadii);
+	FEllipsoid(const FVector& InRadii);
 
-	FVector3d Radii;
-	FVector3d RadiiSquared;
-	FVector3d OneOverRadii;
-	FVector3d OneOverRadiiSquared;
+	FVector Radii;
+	FVector RadiiSquared;
+	FVector OneOverRadii;
+	FVector OneOverRadiiSquared;
 
-	FVector3d GeodeticSurfaceNormal(const FCartesianCoordinates& ECEFLocation) const;
-	FVector3d GeodeticSurfaceNormal(const FGeographicCoordinates& GeographicCoordinates) const;    
+	FVector GeodeticSurfaceNormal(const FCartesianCoordinates& ECEFLocation) const;
+	FVector GeodeticSurfaceNormal(const FGeographicCoordinates& GeographicCoordinates) const;    
 };
diff --git a/Engine/Plugins/Runtime/GeoReferencing/Source/GeoReferencing/Public/GeoReferencingSystem.h b/Engine/Plugins/Runtime/GeoReferencing/Source/GeoReferencing/Public/GeoReferencingSystem.h
--- a/Engine/Plugins/Runtime/GeoReferencing/Source/GeoReferencing/Public/GeoReferencingSystem.h
+++ b/Engine/Plugins/Runtime/GeoReferencing/Source/GeoReferencing/Public/GeoReferencingSystem.h
@@ -12,11 +12,6 @@
 #include "CartesianCoordinates.h"
 #include "Ellipsoid.h"
 
-// Using double precision maths from the Geometry plugin. 
-#include "VectorTypes.h"
-#include "MatrixTypes.h"
-#include "Matrix4d.h"
-
 #include "GeoReferencingSystem.generated.h"
 
 DECLARE_LOG_CATEGORY_EXTERN(LogGeoReferencing, Log, All);
@@ -210,7 +205,7 @@ public:
 	 *    CRS can be identified by their code (EPSG:4326), a well-known text(WKT) string, or PROJ strings...
 	 **/
 	UPROPERTY(EditAnywhere, BlueprintReadWrite, Config, Category = "GeoReferencing")
-	FString ProjectedCRS = FString(TEXT("EPSG:32617")); // UTM Zone 17 North - (Cary) https://epsg.io/32617
+	FString ProjectedCRS = FString(TEXT("EPSG:32631")); // UTM Zone 31 North - (0,0) https://epsg.io/32631
 
 	/**
 	 * String that describes the GEOGRAPHIC CRS of choice.
@@ -245,37 +240,38 @@ public:
 	 * Latitude of UE Origin on planet
 	 **/
 	UPROPERTY(EditAnywhere, BlueprintReadWrite, Category = "GeoReferencing|Origin Location", meta = (EditConditionHides, EditCondition = "!bOriginLocationInProjectedCRS && !bOriginAtPlanetCenter", ClampMin = "-90", ClampMax = "90"))
-	int32 OriginLatitude = 35; // Around Epic Games HQ
+	double OriginLatitude = 0.0;
+
 
 	/**
 	 * Longitude of UE Origin on planet
 	 **/
 	UPROPERTY(EditAnywhere, BlueprintReadWrite, Category = "GeoReferencing|Origin Location", meta = (EditConditionHides, EditCondition = "!bOriginLocationInProjectedCRS && !bOriginAtPlanetCenter", ClampMin = "-180", ClampMax = "180"))
-	int32 OriginLongitude = -78; // Around Epic Games HQ
+	double OriginLongitude = 0.0;
 
 	/**
 	 * Altitude of UE Origin on planet
 	 **/
 	UPROPERTY(EditAnywhere, BlueprintReadWrite, Category = "GeoReferencing|Origin Location", meta = (EditConditionHides, EditCondition = "!bOriginLocationInProjectedCRS && !bOriginAtPlanetCenter"))
-	int32 OriginAltitude = 0; // For constituency, but should be 0
+	double OriginAltitude = 0.0; // For constituency, but should be 0
 
 	/**
 	 * Easting position of UE Origin on planet, express in the Projected CRS Frame
 	 **/
 	UPROPERTY(EditAnywhere, BlueprintReadWrite, Category = "GeoReferencing|Origin Location", meta = (EditConditionHides, EditCondition = "bOriginLocationInProjectedCRS && !bOriginAtPlanetCenter"))
-	int32 OriginProjectedCoordinatesEasting = 705000; // Around Epic Games HQ
+	double OriginProjectedCoordinatesEasting = 500000.0;
 
 	/**
 	 * Northing position of UE Origin on planet, express in the Projected CRS Frame
 	 **/
 	UPROPERTY(EditAnywhere, BlueprintReadWrite, Category = "GeoReferencing|Origin Location", meta = (EditConditionHides, EditCondition = "bOriginLocationInProjectedCRS && !bOriginAtPlanetCenter"))
-	int32 OriginProjectedCoordinatesNorthing = 3960000; // Around Epic Games HQ
+	double OriginProjectedCoordinatesNorthing = 5000000.0;
 
 	/**
 	 * Up position of UE Origin on planet, express in the Projected CRS Frame
 	 **/
 	UPROPERTY(EditAnywhere, BlueprintReadWrite, Category = "GeoReferencing|Origin Location", meta = (EditConditionHides, EditCondition = "bOriginLocationInProjectedCRS && !bOriginAtPlanetCenter"))
-	int32 OriginProjectedCoordinatesUp = 0;
+	double OriginProjectedCoordinatesUp = 0.0;
 
 	
 #if WITH_EDITOR
diff --git a/Engine/Plugins/Runtime/GeoReferencing/Source/GeoReferencing/Public/GeographicCoordinates.h b/Engine/Plugins/Runtime/GeoReferencing/Source/GeoReferencing/Public/GeographicCoordinates.h
--- a/Engine/Plugins/Runtime/GeoReferencing/Source/GeoReferencing/Public/GeographicCoordinates.h
+++ b/Engine/Plugins/Runtime/GeoReferencing/Source/GeoReferencing/Public/GeographicCoordinates.h
@@ -17,8 +17,11 @@ public:
 	FGeographicCoordinates();
 	FGeographicCoordinates(double InLongitude, double InLatitude, double InAltitude);
 
+	UPROPERTY(EditAnywhere, BlueprintReadWrite, Config, Category = "GeoReferencing")
 	double Longitude;
+	UPROPERTY(EditAnywhere, BlueprintReadWrite, Config, Category = "GeoReferencing")
 	double Latitude;
+	UPROPERTY(EditAnywhere, BlueprintReadWrite, Config, Category = "GeoReferencing")
 	double Altitude;
 
 	FText ToFullText(int32 IntegralDigitsLatLon = 8, int32 IntegralDigitsAlti = 2, bool bAsDMS = false);
diff --git a/Engine/Plugins/Runtime/GeoReferencing/Source/GeoReferencing/Public/Matrix4d.h b/Engine/Plugins/Runtime/GeoReferencing/Source/GeoReferencing/Public/Matrix4d.h
deleted file mode 100644
diff --git a/Engine/Plugins/Runtime/GeoReferencing/Source/GeoReferencingEditor/Private/GeoReferencingEditorBPLibrary.cpp b/Engine/Plugins/Runtime/GeoReferencing/Source/GeoReferencingEditor/Private/GeoReferencingEditorBPLibrary.cpp
--- a/Engine/Plugins/Runtime/GeoReferencing/Source/GeoReferencingEditor/Private/GeoReferencingEditorBPLibrary.cpp
+++ b/Engine/Plugins/Runtime/GeoReferencing/Source/GeoReferencingEditor/Private/GeoReferencingEditorBPLibrary.cpp
@@ -89,7 +89,6 @@ void UGeoReferencingEditorBPLibrary::LineTrace(const FVector WorldLocation, cons
 	FCollisionObjectQueryParams ObjectParams = FCollisionObjectQueryParams(ECC_WorldStatic);
 	ObjectParams.AddObjectTypesToQuery(ECC_WorldDynamic);
 	ObjectParams.AddObjectTypesToQuery(ECC_Pawn);
-	ObjectParams.AddObjectTypesToQuery(ECC_Visibility);
 
 	if (GWorld->LineTraceSingleByObjectType(HitResult, LineCheckStart, LineCheckEnd, ObjectParams, CollisionParams))
 	{
