Upgrade of Georeferencing plugin to LWC - Deprecating types, and extra LWC-Related tests

### Notes

Deprecated the double precision FCartesianCoordinates class which is now replaced by FVector
- Kept the C++/BP prototypes for compatibility (The plugin was published as production-ready)
- Added new prototypes with FVector, while keeping the name (K2_)
- Updated all the provided BP to the new definitions with FVectors

Added some convenience methods to convert directly from Engine to Geographic. (Depending on the planet shape, some intermediate paths could fail)
Added a new Round-Planet Pawn to allow for Ellipsoidal movement around the Planet, and keep the horizon horizontal.
Added a sample for a Latitude/Longitude/Altitude placed actor.
Added Touch Controls for easier tests on mobile platforms.

Updated the GeoReferencingDemo test application to work in LWC with the new types

Known issues :
- Without increasing the WORLD_MAX constant, the ellipsoid will be clipped. Use r.UseVisibilityOctree 0 to disable the bad culling
- Enable World Bounds Checks should be disabled in WorldSettings, otherwise, the pawn keeps falling down Z-Axis.

### Hashtags

diff --git a/Engine/Plugins/Runtime/GeoReferencing/Source/GeoReferencing/GeoReferencing.Build.cs b/Engine/Plugins/Runtime/GeoReferencing/Source/GeoReferencing/GeoReferencing.Build.cs
--- a/Engine/Plugins/Runtime/GeoReferencing/Source/GeoReferencing/GeoReferencing.Build.cs
+++ b/Engine/Plugins/Runtime/GeoReferencing/Source/GeoReferencing/GeoReferencing.Build.cs
@@ -24,6 +24,7 @@ public class GeoReferencing : ModuleRules
 			new string[]
 			{
 				"CoreUObject",
+				"InputCore",
 				"Engine",
 				"PROJ",
 				"Projects",
diff --git a/Engine/Plugins/Runtime/GeoReferencing/Source/GeoReferencing/Private/Ellipsoid.cpp b/Engine/Plugins/Runtime/GeoReferencing/Source/GeoReferencing/Private/Ellipsoid.cpp
--- a/Engine/Plugins/Runtime/GeoReferencing/Source/GeoReferencing/Private/Ellipsoid.cpp
+++ b/Engine/Plugins/Runtime/GeoReferencing/Source/GeoReferencing/Private/Ellipsoid.cpp
@@ -4,7 +4,7 @@
 #include "MathUtil.h"
 
 
-// TODOLWC - To be replaced once FVector::Normalize will use a smaller number than 1e-8
+// LWC_TODO - To be replaced once FVector::Normalize will use a smaller number than 1e-8
 #define GEOREF_DOUBLE_SMALL_NUMBER (1.e-50)
 
 
@@ -30,7 +30,17 @@ FEllipsoid::FEllipsoid(const FVector& InRadii)
 	OneOverRadiiSquared = FVector(1.0 / (InRadii.X * InRadii.X), 1.0 / (InRadii.Y * InRadii.Y), 1.0 / (InRadii.Z * InRadii.Z));
 }
 
-FVector FEllipsoid::GeodeticSurfaceNormal(const FCartesianCoordinates& ECEFLocation) const
+double FEllipsoid::GetMaximumRadius()
+{
+	return FMathd::Max3(Radii.X,Radii.Y, Radii.Z);
+}
+
+double FEllipsoid::GetMinimumRadius()
+{
+	return FMathd::Min3(Radii.X, Radii.Y, Radii.Z);
+}
+
+FVector FEllipsoid::GeodeticSurfaceNormal(const FVector& ECEFLocation) const
 {
 	FVector Normal( ECEFLocation.X * OneOverRadiiSquared.X, ECEFLocation.Y * OneOverRadiiSquared.Y, ECEFLocation.Z * OneOverRadiiSquared.Z);
 	Normal.Normalize(GEOREF_DOUBLE_SMALL_NUMBER); 
diff --git a/Engine/Plugins/Runtime/GeoReferencing/Source/GeoReferencing/Private/GeoReferencing.cpp b/Engine/Plugins/Runtime/GeoReferencing/Source/GeoReferencing/Private/GeoReferencing.cpp
deleted file mode 100644
diff --git a/Engine/Plugins/Runtime/GeoReferencing/Source/GeoReferencing/Private/GeoReferencingBFL.cpp b/Engine/Plugins/Runtime/GeoReferencing/Source/GeoReferencing/Private/GeoReferencingBFL.cpp
new file mode 100644
--- /dev/null
+++ b/Engine/Plugins/Runtime/GeoReferencing/Source/GeoReferencing/Private/GeoReferencingBFL.cpp
@@ -0,0 +1,43 @@
+// Copyright Epic Games, Inc. All Rights Reserved.
+
+
+#include "GeoReferencingBFL.h"
+
+FText UGeoReferencingBFL::ToFullText(FVector& CartesianCoordinates, int32 IntegralDigits /*= 3*/)
+{
+	FNumberFormattingOptions NumberFormatOptions = FNumberFormattingOptions::DefaultNoGrouping();
+	NumberFormatOptions.MinimumFractionalDigits = IntegralDigits;
+	NumberFormatOptions.MaximumFractionalDigits = IntegralDigits;
+
+	FFormatNamedArguments Args;
+	Args.Add(TEXT("X"), FText::AsNumber(CartesianCoordinates.X, &NumberFormatOptions));
+	Args.Add(TEXT("Y"), FText::AsNumber(CartesianCoordinates.Y, &NumberFormatOptions));
+	Args.Add(TEXT("Z"), FText::AsNumber(CartesianCoordinates.Z, &NumberFormatOptions));
+
+	return FText::Format(NSLOCTEXT("GeoReferencing", "CartesianCoordinatesFull", "Easting={X} Northing={Y} Up={Z}"), Args);
+}
+
+FText UGeoReferencingBFL::ToCompactText(FVector& CartesianCoordinates, int32 IntegralDigits /*= 3*/)
+{
+	FNumberFormattingOptions NumberFormatOptions = FNumberFormattingOptions::DefaultNoGrouping();
+	NumberFormatOptions.MinimumFractionalDigits = IntegralDigits;
+	NumberFormatOptions.MaximumFractionalDigits = IntegralDigits;
+
+	FFormatNamedArguments Args;
+	Args.Add(TEXT("X"), FText::AsNumber(CartesianCoordinates.X, &NumberFormatOptions));
+	Args.Add(TEXT("Y"), FText::AsNumber(CartesianCoordinates.Y, &NumberFormatOptions));
+	Args.Add(TEXT("Z"), FText::AsNumber(CartesianCoordinates.Z, &NumberFormatOptions));
+
+	return FText::Format(NSLOCTEXT("GeoReferencing", "CartesianCoordinatesCompact", "({X}, {Y}, {Z}m )"), Args);
+}
+
+void UGeoReferencingBFL::ToSeparateTexts(FVector& CartesianCoordinates, FText& OutX, FText& OutY, FText& OutZ, int32 IntegralDigits /*= 3*/)
+{
+	FNumberFormattingOptions NumberFormatOptions = FNumberFormattingOptions::DefaultNoGrouping();
+	NumberFormatOptions.MinimumFractionalDigits = IntegralDigits;
+	NumberFormatOptions.MaximumFractionalDigits = IntegralDigits;
+
+	OutX = FText::AsNumber(CartesianCoordinates.X, &NumberFormatOptions);
+	OutY = FText::AsNumber(CartesianCoordinates.Y, &NumberFormatOptions);
+	OutZ = FText::AsNumber(CartesianCoordinates.Z, &NumberFormatOptions);
+}
diff --git a/Engine/Plugins/Runtime/GeoReferencing/Source/GeoReferencing/Private/GeoReferencingModule.cpp b/Engine/Plugins/Runtime/GeoReferencing/Source/GeoReferencing/Private/GeoReferencingModule.cpp
new file mode 100644
--- /dev/null
+++ b/Engine/Plugins/Runtime/GeoReferencing/Source/GeoReferencing/Private/GeoReferencingModule.cpp
@@ -0,0 +1,7 @@
+// Copyright Epic Games, Inc. All Rights Reserved.
+
+#include "GeoReferencingModule.h"
+
+GEOREFERENCING_API DEFINE_LOG_CATEGORY(LogGeoReferencing);
+
+IMPLEMENT_MODULE(FGeoReferencingModule, GeoReferencing)
diff --git a/Engine/Plugins/Runtime/GeoReferencing/Source/GeoReferencing/Private/GeoReferencingSystem.cpp b/Engine/Plugins/Runtime/GeoReferencing/Source/GeoReferencing/Private/GeoReferencingSystem.cpp
--- a/Engine/Plugins/Runtime/GeoReferencing/Source/GeoReferencing/Private/GeoReferencingSystem.cpp
+++ b/Engine/Plugins/Runtime/GeoReferencing/Source/GeoReferencing/Private/GeoReferencingSystem.cpp
@@ -10,7 +10,7 @@
 #include "Engine/Engine.h"
 #include "Engine/World.h"
 #include "MathUtil.h"
-
+ 
 #include "HAL/PlatformFileManager.h"
 #include "IPlatformFilePak.h"
 
@@ -23,11 +23,11 @@ THIRD_PARTY_INCLUDES_END
 
 #define ECEF_EPSG_FSTRING FString(TEXT("EPSG:4978"))
 
-// TODOLWC - To be replaced once FVector::Normalize will use a smaller number than 1e-8
+// LWC_TODO - To be replaced once FVector::Normalize will use a smaller number than 1e-8
 #define GEOREF_DOUBLE_SMALL_NUMBER			(1.e-50)
 
 
-DEFINE_LOG_CATEGORY(LogGeoReferencing);
+
 
 AGeoReferencingSystem* AGeoReferencingSystem::GetGeoReferencingSystem(UObject* WorldContextObject)
 {
@@ -71,7 +71,8 @@ public:
 	void DeInitPROJLibrary();
 	PJ* GetPROJProjection(FString SourceCRS, FString DestinationCRS);
 	bool GetEllipsoid(FString CRSString, FEllipsoid& Ellipsoid);
-	FMatrix GetWorldFrameToECEFFrame(const FEllipsoid& Ellipsoid, const FCartesianCoordinates& ECEFLocation);
+	
+	FMatrix GetWorldFrameToECEFFrame(const FEllipsoid& Ellipsoid, const FVector& ECEFLocation);
 
 	PJ_CONTEXT* ProjContext;
 	PJ* ProjProjectedToGeographic;
@@ -114,7 +115,7 @@ void AGeoReferencingSystem::Initialize()
 
 	Impl->InitPROJLibrary();
 
-	// TODOABLWC - Can we consider other conventions ? Or North Offset like in SunPosition?
+	// Should we consider other conventions ? Or North Offset like in SunPosition?
 	Impl->WorldFrameToUEFrame = FMatrix( 
 		FVector(1.0, 0.0, 0.0),		// Easting (X) is UE World X
 		FVector(0.0, -1.0, 0.0),	// Northing (Y) is UE World -Y because of left-handed convention
@@ -137,14 +138,106 @@ void AGeoReferencingSystem::BeginDestroy()
 	}
 }
 
+#pragma region Old deprecated Prototypes
+
 void AGeoReferencingSystem::EngineToProjected(const FVector& EngineCoordinates, FCartesianCoordinates& ProjectedCoordinates)
+{
+	FVector Result;
+	EngineToProjected(EngineCoordinates, Result);
+	ProjectedCoordinates = FCartesianCoordinates(Result);
+}
+
+void AGeoReferencingSystem::ProjectedToEngine(const FCartesianCoordinates& ProjectedCoordinates, FVector& EngineCoordinates)
+{
+	ProjectedToEngine(ProjectedCoordinates.ToVector(), EngineCoordinates);
+}
+
+void AGeoReferencingSystem::EngineToECEF(const FVector& EngineCoordinates, FCartesianCoordinates& ECEFCoordinates)
+{
+	FVector Result;
+	EngineToECEF(EngineCoordinates, Result);
+	ECEFCoordinates = FCartesianCoordinates(Result);
+}
+
+void AGeoReferencingSystem::ECEFToEngine(const FCartesianCoordinates& ECEFCoordinates, FVector& EngineCoordinates)
+{
+	ECEFToEngine(ECEFCoordinates.ToVector(), EngineCoordinates);
+}
+
+
+void AGeoReferencingSystem::ProjectedToGeographic(const FCartesianCoordinates& ProjectedCoordinates, FGeographicCoordinates& GeographicCoordinates)
+{
+	ProjectedToGeographic(ProjectedCoordinates.ToVector(), GeographicCoordinates);
+}
+
+void AGeoReferencingSystem::GeographicToProjected(const FGeographicCoordinates& GeographicCoordinates, FCartesianCoordinates& ProjectedCoordinates)
+{
+	FVector Result;
+	GeographicToProjected(GeographicCoordinates, Result);
+	ProjectedCoordinates = FCartesianCoordinates(Result);
+}
+
+void AGeoReferencingSystem::ProjectedToECEF(const FCartesianCoordinates& ProjectedCoordinates, FCartesianCoordinates& ECEFCoordinates)
+{
+	FVector Result;
+	ProjectedToECEF(ProjectedCoordinates.ToVector(), Result);
+	ECEFCoordinates = FCartesianCoordinates(Result);
+}
+
+void AGeoReferencingSystem::ECEFToProjected(const FCartesianCoordinates& ECEFCoordinates, FCartesianCoordinates& ProjectedCoordinates)
+{
+	FVector Result;
+	ECEFToProjected(ECEFCoordinates.ToVector(), Result);
+	ProjectedCoordinates = FCartesianCoordinates(Result);
+}
+
+void AGeoReferencingSystem::GeographicToECEF(const FGeographicCoordinates& GeographicCoordinates, FCartesianCoordinates& ECEFCoordinates)
+{
+	FVector Result;
+	GeographicToECEF(GeographicCoordinates, Result);
+	ECEFCoordinates = FCartesianCoordinates(Result);
+}
+
+void AGeoReferencingSystem::ECEFToGeographic(const FCartesianCoordinates& ECEFCoordinates, FGeographicCoordinates& GeographicCoordinates)
+{
+	ECEFToGeographic(ECEFCoordinates.ToVector(), GeographicCoordinates);
+}
+
+void AGeoReferencingSystem::GetENUVectorsAtProjectedLocation(const FCartesianCoordinates& ProjectedCoordinates, FVector& East, FVector& North, FVector& Up)
+{
+	GetENUVectorsAtProjectedLocation(ProjectedCoordinates.ToVector(), East, North, Up);
+}
+
+void AGeoReferencingSystem::GetENUVectorsAtECEFLocation(const FCartesianCoordinates& ECEFCoordinates, FVector& East, FVector& North, FVector& Up)
+{
+	GetENUVectorsAtECEFLocation(ECEFCoordinates.ToVector(), East, North, Up);
+}
+
+void AGeoReferencingSystem::GetECEFENUVectorsAtECEFLocation(const FCartesianCoordinates& ECEFCoordinates, FVector& ECEFEast, FVector& ECEFNorth, FVector& ECEFUp)
+{
+	GetECEFENUVectorsAtECEFLocation(ECEFCoordinates.ToVector(), ECEFEast, ECEFNorth, ECEFUp);
+}
+
+FTransform AGeoReferencingSystem::GetTangentTransformAtProjectedLocation(const FCartesianCoordinates& ProjectedCoordinates)
+{
+	return GetTangentTransformAtProjectedLocation(ProjectedCoordinates.ToVector());
+}
+
+FTransform AGeoReferencingSystem::GetTangentTransformAtECEFLocation(const FCartesianCoordinates& ECEFCoordinates)
+{
+	return GetTangentTransformAtECEFLocation(ECEFCoordinates.ToVector());
+}
+#pragma endregion
+
+
+void AGeoReferencingSystem::EngineToProjected(const FVector& EngineCoordinates, FVector& ProjectedCoordinates)
 {
 	switch (PlanetShape)
 	{
 	case EPlanetShape::RoundPlanet:
 	{
 		// In RoundPlanet, we have to go through the ECEF transform as an intermediate step. 
-		FCartesianCoordinates ECEFCoordinates;
+		FVector ECEFCoordinates;
 		EngineToECEF(EngineCoordinates, ECEFCoordinates);
 		ECEFToProjected(ECEFCoordinates, ProjectedCoordinates);
 	}
@@ -165,21 +258,20 @@ void AGeoReferencingSystem::EngineToProjected(const FVector& EngineCoordinates,
 		FVector UEWorldCoordinates = UERebasedCoordinates * FVector(0.01, -0.01, 0.01);
 
 		// Add the defined origin offset
-		FVector Projected = UEWorldCoordinates + Impl->WorldOriginLocationProjected;
-		ProjectedCoordinates = FCartesianCoordinates(Projected);
+		ProjectedCoordinates = UEWorldCoordinates + Impl->WorldOriginLocationProjected;
 	}
 	break;
 	}
 }
 
-void AGeoReferencingSystem::ProjectedToEngine(const FCartesianCoordinates& ProjectedCoordinates, FVector& EngineCoordinates)
+void AGeoReferencingSystem::ProjectedToEngine(const FVector& ProjectedCoordinates, FVector& EngineCoordinates)
 {
 	switch (PlanetShape)
 	{
 	case EPlanetShape::RoundPlanet:
 	{
 		// In RoundPlanet, we have to go through the ECEF transform as an intermediate step. 
-		FCartesianCoordinates ECEFCoordinates;
+		FVector ECEFCoordinates;
 		ProjectedToECEF(ProjectedCoordinates, ECEFCoordinates);
 		ECEFToEngine(ECEFCoordinates, EngineCoordinates);
 	}
@@ -190,23 +282,20 @@ void AGeoReferencingSystem::ProjectedToEngine(const FCartesianCoordinates& Proje
 	{
 		// in FlatPlanet, the transform is simply a translation
 		// Remove the Origin location, and convert to UE Units, while inverting the Z Axis
-		FVector Projected(ProjectedCoordinates.X, ProjectedCoordinates.Y, ProjectedCoordinates.Z);
-		FVector UEWorldCoordinates = (Projected - Impl->WorldOriginLocationProjected);
+		FVector UEWorldCoordinates = (ProjectedCoordinates - Impl->WorldOriginLocationProjected);
 
 		// Convert UE units to meters, invert the Y coordinate because of left-handed UE Frame
 		FVector UERebasedCoordinates = UEWorldCoordinates * FVector(100.0, -100.0, 100.0);
 
 		// Consider the UE internal rebasing
 		FIntVector UERebasingOffset = GetWorld()->OriginLocation;
-		FVector EngineCoordinates3d = UERebasedCoordinates - FVector(UERebasingOffset.X, UERebasingOffset.Y, UERebasingOffset.Z);
-
-		EngineCoordinates = FVector(EngineCoordinates3d.X, EngineCoordinates3d.Y, EngineCoordinates3d.Z);
+		EngineCoordinates = UERebasedCoordinates - FVector(UERebasingOffset.X, UERebasingOffset.Y, UERebasingOffset.Z);
 	}
 	break;
 	}
 }
 
-void AGeoReferencingSystem::EngineToECEF(const FVector& EngineCoordinates, FCartesianCoordinates& ECEFCoordinates)
+void AGeoReferencingSystem::EngineToECEF(const FVector& EngineCoordinates, FVector& ECEFCoordinates)
 {
 	switch (PlanetShape)
 	{
@@ -229,8 +318,7 @@ void AGeoReferencingSystem::EngineToECEF(const FVector& EngineCoordinates, FCart
 		}
 		else
 		{
-			FVector ECEFCoordinates4d = Impl->WorldFrameToECEFFrame.TransformPosition( FVector(UEWorldCoordinates.X, UEWorldCoordinates.Y, UEWorldCoordinates.Z));
-			ECEFCoordinates = FCartesianCoordinates(ECEFCoordinates4d);
+			ECEFCoordinates = Impl->WorldFrameToECEFFrame.TransformPosition(UEWorldCoordinates);
 		}
 	}
 	break;
@@ -239,7 +327,7 @@ void AGeoReferencingSystem::EngineToECEF(const FVector& EngineCoordinates, FCart
 	default:
 	{
 		// In FlatPlanet, we have to go through the Projected transform as an intermediate step. 
-		FCartesianCoordinates ProjectedCoordinates;
+		FVector ProjectedCoordinates;
 		EngineToProjected(EngineCoordinates, ProjectedCoordinates);
 		ProjectedToECEF(ProjectedCoordinates, ECEFCoordinates);
 	}
@@ -247,7 +335,7 @@ void AGeoReferencingSystem::EngineToECEF(const FVector& EngineCoordinates, FCart
 	}
 }
 
-void AGeoReferencingSystem::ECEFToEngine(const FCartesianCoordinates& ECEFCoordinates, FVector& EngineCoordinates)
+void AGeoReferencingSystem::ECEFToEngine(const FVector& ECEFCoordinates, FVector& EngineCoordinates)
 {
 	switch (PlanetShape)
 	{
@@ -257,11 +345,11 @@ void AGeoReferencingSystem::ECEFToEngine(const FCartesianCoordinates& ECEFCoordi
 		if (bOriginAtPlanetCenter)
 		{
 			// Easy case, UE is ECEF... And we did the rebasing, so return the Global coordinates
-			UEWorldCoordinates = FVector(ECEFCoordinates.X, ECEFCoordinates.Y, ECEFCoordinates.Z);
+			UEWorldCoordinates = ECEFCoordinates;
 		}
 		else
 		{
-			UEWorldCoordinates = Impl->ECEFFrameToWorldFrame.TransformPosition(FVector(ECEFCoordinates.X, ECEFCoordinates.Y, ECEFCoordinates.Z));
+			UEWorldCoordinates = Impl->ECEFFrameToWorldFrame.TransformPosition(ECEFCoordinates);
 		}
 
 		// Convert UE units to meters, invert the Y coordinate because of left-handed UE Frame
@@ -269,9 +357,7 @@ void AGeoReferencingSystem::ECEFToEngine(const FCartesianCoordinates& ECEFCoordi
 
 		// Consider the UE internal rebasing
 		FIntVector UERebasingOffset = GetWorld()->OriginLocation;
-		FVector EngineCoordinates3d = UERebasedCoordinates - FVector(UERebasingOffset.X, UERebasingOffset.Y, UERebasingOffset.Z);
-
-		EngineCoordinates = FVector(EngineCoordinates3d.X, EngineCoordinates3d.Y, EngineCoordinates3d.Z);
+		EngineCoordinates = UERebasedCoordinates - FVector(UERebasingOffset.X, UERebasingOffset.Y, UERebasingOffset.Z);
 	}
 	break;
 
@@ -279,7 +365,7 @@ void AGeoReferencingSystem::ECEFToEngine(const FCartesianCoordinates& ECEFCoordi
 	default:
 	{
 		// In FlatPlanet, we have to go through the Projected transform as an intermediate step. 
-		FCartesianCoordinates ProjectedCoordinates;
+		FVector ProjectedCoordinates;
 		ECEFToProjected(ECEFCoordinates, ProjectedCoordinates);
 		ProjectedToEngine(ProjectedCoordinates, EngineCoordinates);
 	}
@@ -287,7 +373,56 @@ void AGeoReferencingSystem::ECEFToEngine(const FCartesianCoordinates& ECEFCoordi
 	}
 }
 
-void AGeoReferencingSystem::ProjectedToGeographic(const FCartesianCoordinates& ProjectedCoordinates, FGeographicCoordinates& GeographicCoordinates)
+
+void AGeoReferencingSystem::EngineToGeographic(const FVector& EngineCoordinates, FGeographicCoordinates& GeographicCoordinates)
+{
+	switch (PlanetShape)
+	{
+		case EPlanetShape::RoundPlanet:
+		{
+			FVector ECEFCoordinates;
+			EngineToECEF(EngineCoordinates, ECEFCoordinates);
+			ECEFToGeographic(ECEFCoordinates, GeographicCoordinates);
+		}
+	break;
+
+		case EPlanetShape::FlatPlanet:
+		default:
+		{
+			FVector ProjectedCoordinates;
+			EngineToProjected(EngineCoordinates, ProjectedCoordinates);
+			ProjectedToGeographic(ProjectedCoordinates, GeographicCoordinates);
+
+		}
+	break;
+	}
+}
+
+void AGeoReferencingSystem::GeographicToEngine(const FGeographicCoordinates& GeographicCoordinates, FVector& EngineCoordinates)
+{
+	switch (PlanetShape)
+	{
+		case EPlanetShape::RoundPlanet:
+		{
+			FVector ECEFCoordinates;
+			GeographicToECEF(GeographicCoordinates, ECEFCoordinates);
+			ECEFToEngine(ECEFCoordinates, EngineCoordinates);
+		}
+	break;
+
+		case EPlanetShape::FlatPlanet:
+		default:
+		{
+			FVector ProjectedCoordinates;
+			GeographicToProjected(GeographicCoordinates,ProjectedCoordinates);
+			ProjectedToEngine(ProjectedCoordinates, EngineCoordinates);
+		}
+	break;
+	}
+}
+
+
+void AGeoReferencingSystem::ProjectedToGeographic(const FVector& ProjectedCoordinates, FGeographicCoordinates& GeographicCoordinates)
 {
 	PJ_COORD input, output;
 	input = proj_coord(ProjectedCoordinates.X, ProjectedCoordinates.Y, ProjectedCoordinates.Z, 0);
@@ -298,7 +433,7 @@ void AGeoReferencingSystem::ProjectedToGeographic(const FCartesianCoordinates& P
 	GeographicCoordinates.Altitude = output.lpz.z;
 }
 
-void AGeoReferencingSystem::GeographicToProjected(const FGeographicCoordinates& GeographicCoordinates, FCartesianCoordinates& ProjectedCoordinates)
+void AGeoReferencingSystem::GeographicToProjected(const FGeographicCoordinates& GeographicCoordinates, FVector& ProjectedCoordinates)
 {
 	PJ_COORD input, output;
 	input = proj_coord(GeographicCoordinates.Longitude, GeographicCoordinates.Latitude, GeographicCoordinates.Altitude, 0);
@@ -309,7 +444,7 @@ void AGeoReferencingSystem::GeographicToProjected(const FGeographicCoordinates&
 	ProjectedCoordinates.Z = output.xyz.z;
 }
 
-void AGeoReferencingSystem::ProjectedToECEF(const FCartesianCoordinates& ProjectedCoordinates, FCartesianCoordinates& ECEFCoordinates)
+void AGeoReferencingSystem::ProjectedToECEF(const FVector& ProjectedCoordinates, FVector& ECEFCoordinates)
 {
 	PJ_COORD input, output;
 	input = proj_coord(ProjectedCoordinates.X, ProjectedCoordinates.Y, ProjectedCoordinates.Z, 0);
@@ -320,7 +455,7 @@ void AGeoReferencingSystem::ProjectedToECEF(const FCartesianCoordinates& Project
 	ECEFCoordinates.Z = output.xyz.z;
 }
 
-void AGeoReferencingSystem::ECEFToProjected(const FCartesianCoordinates& ECEFCoordinates, FCartesianCoordinates& ProjectedCoordinates)
+void AGeoReferencingSystem::ECEFToProjected(const FVector& ECEFCoordinates, FVector& ProjectedCoordinates)
 {
 	PJ_COORD input, output;
 	input = proj_coord(ECEFCoordinates.X, ECEFCoordinates.Y, ECEFCoordinates.Z, 0);
@@ -331,7 +466,7 @@ void AGeoReferencingSystem::ECEFToProjected(const FCartesianCoordinates& ECEFCoo
 	ProjectedCoordinates.Z = output.xyz.z;
 }
 
-void AGeoReferencingSystem::GeographicToECEF(const FGeographicCoordinates& GeographicCoordinates, FCartesianCoordinates& ECEFCoordinates)
+void AGeoReferencingSystem::GeographicToECEF(const FGeographicCoordinates& GeographicCoordinates, FVector& ECEFCoordinates)
 {
 	PJ_COORD input, output;
 	input = proj_coord(GeographicCoordinates.Longitude, GeographicCoordinates.Latitude, GeographicCoordinates.Altitude, 0);
@@ -342,7 +477,7 @@ void AGeoReferencingSystem::GeographicToECEF(const FGeographicCoordinates& Geogr
 	ECEFCoordinates.Z = output.xyz.z;
 }
 
-void AGeoReferencingSystem::ECEFToGeographic(const FCartesianCoordinates& ECEFCoordinates, FGeographicCoordinates& GeographicCoordinates)
+void AGeoReferencingSystem::ECEFToGeographic(const FVector& ECEFCoordinates, FGeographicCoordinates& GeographicCoordinates)
 {
 	PJ_COORD input, output;
 	input = proj_coord(ECEFCoordinates.X, ECEFCoordinates.Y, ECEFCoordinates.Z, 0);
@@ -357,26 +492,26 @@ void AGeoReferencingSystem::ECEFToGeographic(const FCartesianCoordinates& ECEFCo
 
 void AGeoReferencingSystem::GetENUVectorsAtEngineLocation(const FVector& EngineCoordinates, FVector& East, FVector& North, FVector& Up)
 {
-	FCartesianCoordinates ECEFLocation;
+	FVector ECEFLocation;
 	EngineToECEF(EngineCoordinates, ECEFLocation);
 	GetENUVectorsAtECEFLocation(ECEFLocation, East, North, Up);
 }
 
-void AGeoReferencingSystem::GetENUVectorsAtProjectedLocation(const FCartesianCoordinates& ProjectedCoordinates, FVector& East, FVector& North, FVector& Up)
+void AGeoReferencingSystem::GetENUVectorsAtProjectedLocation(const FVector& ProjectedCoordinates, FVector& East, FVector& North, FVector& Up)
 {
-	FCartesianCoordinates ECEFLocation;
+	FVector ECEFLocation;
 	ProjectedToECEF(ProjectedCoordinates, ECEFLocation);
 	GetENUVectorsAtECEFLocation(ECEFLocation, East, North, Up);
 }
 
 void AGeoReferencingSystem::GetENUVectorsAtGeographicLocation(const FGeographicCoordinates& GeographicCoordinates, FVector& East, FVector& North, FVector& Up)
 {
-	FCartesianCoordinates ECEFLocation;
+	FVector ECEFLocation;
 	GeographicToECEF(GeographicCoordinates, ECEFLocation);
 	GetENUVectorsAtECEFLocation(ECEFLocation, East, North, Up);
 }
 
-void AGeoReferencingSystem::GetENUVectorsAtECEFLocation(const FCartesianCoordinates& ECEFCoordinates, FVector& East, FVector& North, FVector& Up)
+void AGeoReferencingSystem::GetENUVectorsAtECEFLocation(const FVector& ECEFCoordinates, FVector& East, FVector& North, FVector& Up)
 {
 	// Compute Tangent matrix at ECEF location
 	FEllipsoid& Ellipsoid = Impl->GeographicEllipsoid;
@@ -400,10 +535,10 @@ void AGeoReferencingSystem::GetENUVectorsAtECEFLocation(const FCartesianCoordina
 	case EPlanetShape::FlatPlanet:
 	default:
 		// PROJ don't provide anything to project direction vectors. Let's do it by hand...
-		FVector EasternPoint = ECEFCoordinates.ToVector() + WorldFrameToECEFFrameAtLocation.TransformVector(FVector(1.0, 0.0, 0.0)); // 1m from origin to the East
-		FVector NorthernPoint = ECEFCoordinates.ToVector() + WorldFrameToECEFFrameAtLocation.TransformVector(FVector(0.0, 1.0, 0.0)); // 1m from origin to the North
+		FVector EasternPoint = ECEFCoordinates + WorldFrameToECEFFrameAtLocation.TransformVector(FVector(1.0, 0.0, 0.0)); // 1m from origin to the East
+		FVector NorthernPoint = ECEFCoordinates + WorldFrameToECEFFrameAtLocation.TransformVector(FVector(0.0, 1.0, 0.0)); // 1m from origin to the North
 		
-		FCartesianCoordinates ProjectedOrigin, ProjectedEastern, ProjectedNorthern;
+		FVector ProjectedOrigin, ProjectedEastern, ProjectedNorthern;
 		ECEFToProjected(ECEFCoordinates, ProjectedOrigin);
 		ECEFToProjected(EasternPoint, ProjectedEastern);
 		ECEFToProjected(NorthernPoint, ProjectedNorthern);
@@ -421,7 +556,7 @@ void AGeoReferencingSystem::GetENUVectorsAtECEFLocation(const FCartesianCoordina
 	}
 }
 
-void AGeoReferencingSystem::GetECEFENUVectorsAtECEFLocation(const FCartesianCoordinates& ECEFCoordinates, FVector& ECEFEast, FVector& ECEFNorth, FVector& ECEFUp)
+void AGeoReferencingSystem::GetECEFENUVectorsAtECEFLocation(const FVector& ECEFCoordinates, FVector& ECEFEast, FVector& ECEFNorth, FVector& ECEFUp)
 {
 	// Compute Tangent matrix at ECEF location
 	FEllipsoid& Ellipsoid = Impl->GeographicEllipsoid;
@@ -438,26 +573,26 @@ void AGeoReferencingSystem::GetECEFENUVectorsAtECEFLocation(const FCartesianCoor
 
 FTransform AGeoReferencingSystem::GetTangentTransformAtEngineLocation(const FVector& EngineCoordinates)
 {
-	FCartesianCoordinates ECEFLocation;
+	FVector ECEFLocation;
 	EngineToECEF(EngineCoordinates, ECEFLocation);
 	return GetTangentTransformAtECEFLocation(ECEFLocation);
 }
 
-FTransform AGeoReferencingSystem::GetTangentTransformAtProjectedLocation(const FCartesianCoordinates& ProjectedCoordinates)
+FTransform AGeoReferencingSystem::GetTangentTransformAtProjectedLocation(const FVector& ProjectedCoordinates)
 {
-	FCartesianCoordinates ECEFLocation;
+	FVector ECEFLocation;
 	ProjectedToECEF(ProjectedCoordinates, ECEFLocation);
 	return GetTangentTransformAtECEFLocation(ECEFLocation);
 }
 
 FTransform AGeoReferencingSystem::GetTangentTransformAtGeographicLocation(const FGeographicCoordinates& GeographicCoordinates)
 {
-	FCartesianCoordinates ECEFLocation;
+	FVector ECEFLocation;
 	GeographicToECEF(GeographicCoordinates, ECEFLocation);
 	return GetTangentTransformAtECEFLocation(ECEFLocation);
 }
 
-FTransform AGeoReferencingSystem::GetTangentTransformAtECEFLocation(const FCartesianCoordinates& ECEFCoordinates)
+FTransform AGeoReferencingSystem::GetTangentTransformAtECEFLocation(const FVector& ECEFCoordinates)
 {
 	if (PlanetShape == EPlanetShape::RoundPlanet)
 	{
@@ -470,7 +605,7 @@ FTransform AGeoReferencingSystem::GetTangentTransformAtECEFLocation(const FCarte
 		FMatrix WorldFrameToECEFFrameAtLocation = Impl->GetWorldFrameToECEFFrame(Ellipsoid, ECEFCoordinates);
 		
 
-		FMatrix UEtoECEF = Impl->UEFrameToWorldFrame * WorldFrameToECEFFrameAtLocation * Impl->ECEFFrameToWorldFrame * Impl->WorldFrameToUEFrame; //TOOD?
+		FMatrix UEtoECEF = Impl->UEFrameToWorldFrame * WorldFrameToECEFFrameAtLocation * Impl->ECEFFrameToWorldFrame * Impl->WorldFrameToUEFrame; 
 		FVector UEOrigin;
 		ECEFToEngine(ECEFCoordinates, UEOrigin);
 		UEtoECEF.SetOrigin(UEOrigin);
@@ -546,6 +681,28 @@ bool AGeoReferencingSystem::IsCRSStringValid(FString CRSString, FString& Error)
 	return true;
 }
 
+// Ellipsoid
+
+double AGeoReferencingSystem::GetGeographicEllipsoidMaxRadius()
+{
+	return Impl->GeographicEllipsoid.GetMaximumRadius();
+}
+
+double AGeoReferencingSystem::GetGeographicEllipsoidMinRadius()
+{
+	return Impl->GeographicEllipsoid.GetMinimumRadius();
+}
+
+double AGeoReferencingSystem::GetProjectedEllipsoidMaxRadius()
+{
+	return Impl->ProjectedEllipsoid.GetMaximumRadius();
+}
+
+double AGeoReferencingSystem::GetProjectedEllipsoidMinRadius()
+{
+	return Impl->ProjectedEllipsoid.GetMinimumRadius();
+}
+
 bool AGeoReferencingSystem::FGeoReferencingSystemInternals::GetEllipsoid(FString CRSString, FEllipsoid& Ellipsoid)
 {
 	FTCHARToUTF8 ConvertCRSString(*CRSString);
@@ -658,11 +815,11 @@ void AGeoReferencingSystem::ApplySettings()
 		{
 			// We need to compute the ENU Vectors at the origin point. For that, the origin has to be expressed first in ECEF. 
 			// Express origin in ECEF, and get the ENU vectors
-			FCartesianCoordinates ECEFOrigin;
+			FVector ECEFOrigin;
 			Impl->WorldFrameToECEFFrame.SetIdentity();
 			if (bOriginLocationInProjectedCRS)
 			{
-				FCartesianCoordinates ProjectedOrigin(
+				FVector ProjectedOrigin(
 					OriginProjectedCoordinatesEasting,
 					OriginProjectedCoordinatesNorthing,
 					OriginProjectedCoordinatesUp);
@@ -692,7 +849,7 @@ void AGeoReferencingSystem::ApplySettings()
 		else
 		{
 			// World origin is expressed using Geographic coordinates. Convert them to the projected CRS to have the offset
-			FCartesianCoordinates OriginProjected;
+			FVector OriginProjected;
 			GeographicToProjected(FGeographicCoordinates(OriginLongitude, OriginLatitude, OriginAltitude), OriginProjected);
 			Impl->WorldOriginLocationProjected = FVector(OriginProjected.X, OriginProjected.Y, OriginProjected.Z);
 		}
@@ -818,14 +975,14 @@ PJ* AGeoReferencingSystem::FGeoReferencingSystemInternals::GetPROJProjection(FSt
 	return P_for_GIS;
 }
 
-FMatrix AGeoReferencingSystem::FGeoReferencingSystemInternals::GetWorldFrameToECEFFrame(const FEllipsoid& Ellipsoid, const FCartesianCoordinates& ECEFLocation)
+FMatrix AGeoReferencingSystem::FGeoReferencingSystemInternals::GetWorldFrameToECEFFrame(const FEllipsoid& Ellipsoid, const FVector& ECEFLocation)
 {
 	// See ECEF standard : https://commons.wikimedia.org/wiki/File:ECEF_ENU_Longitude_Latitude_right-hand-rule.svg
 	if (FMathd::Abs(ECEFLocation.X) < FMathd::Epsilon &&
 		FMathd::Abs(ECEFLocation.Y) < FMathd::Epsilon)
 	{
 		// Special Case - On earth axis... 
-		int Sign = 1;
+		double Sign = 1.0;
 		if (FMathd::Abs(ECEFLocation.Z) < FMathd::Epsilon)
 		{
 			// At origin - Should not happen, but consider it's the same as north pole
@@ -841,15 +998,15 @@ FMatrix AGeoReferencingSystem::FGeoReferencingSystemInternals::GetWorldFrameToEC
 			FVector::YAxisVector, 			// East = Y
 			-FVector::XAxisVector * Sign,	// North = Sign * X
 			FVector::ZAxisVector * Sign,	// Up = Sign*Z
-			ECEFLocation.ToVector());
+			ECEFLocation);
 	}
 	else
 	{
-		FVector Up = Ellipsoid.GeodeticSurfaceNormal(ECEFLocation); //TODOABLWC - Mauvaises valeurs ici... 
+		FVector Up = Ellipsoid.GeodeticSurfaceNormal(ECEFLocation);
 		FVector East(-ECEFLocation.Y, ECEFLocation.X, 0.0); 
 		East.Normalize(GEOREF_DOUBLE_SMALL_NUMBER);
 		FVector North = Up.Cross(East);
-		return FMatrix(	East, North, Up, ECEFLocation.ToVector());
+		return FMatrix(	East, North, Up, ECEFLocation);
 	}
 }
 
diff --git a/Engine/Plugins/Runtime/GeoReferencing/Source/GeoReferencing/Private/GeographicCoordinates.cpp b/Engine/Plugins/Runtime/GeoReferencing/Source/GeoReferencing/Private/GeographicCoordinates.cpp
--- a/Engine/Plugins/Runtime/GeoReferencing/Source/GeoReferencing/Private/GeographicCoordinates.cpp
+++ b/Engine/Plugins/Runtime/GeoReferencing/Source/GeoReferencing/Private/GeographicCoordinates.cpp
@@ -1,6 +1,7 @@
 // Copyright Epic Games, Inc. All Rights Reserved.
 
 #include "GeographicCoordinates.h"
+#include "MathUtil.h"
 
 FGeographicCoordinates::FGeographicCoordinates()
 	: Longitude(0)
@@ -9,9 +10,16 @@ FGeographicCoordinates::FGeographicCoordinates()
 {}
 
 FGeographicCoordinates::FGeographicCoordinates(double InLongitude, double InLatitude, double InAltitude)
-	: Longitude(InLongitude)
-	, Latitude(InLatitude)
+	: Longitude(FMathd::Clamp(InLongitude, -180.0, 180.0))
+	, Latitude(FMathd::Clamp(InLatitude, -90.0, 90.0))
 	, Altitude(InAltitude)
+{
+}
+
+FGeographicCoordinates::FGeographicCoordinates(const FVector& LatLongAltVector)
+	: Longitude(FMathd::Clamp(LatLongAltVector.Y, -180.0, 180.0))
+	, Latitude(FMathd::Clamp(LatLongAltVector.X, -90.0, 90.0))
+	, Altitude(LatLongAltVector.Z)
 {}
 
 FText FGeographicCoordinates::ToFullText(int32 IntegralDigitsLatLon /*= 8*/, int32 IntegralDigitsAlti /*= 2*/, bool bAsDMS /*= false*/)
diff --git a/Engine/Plugins/Runtime/GeoReferencing/Source/GeoReferencing/Private/RoundPlanetPawn.cpp b/Engine/Plugins/Runtime/GeoReferencing/Source/GeoReferencing/Private/RoundPlanetPawn.cpp
new file mode 100644
--- /dev/null
+++ b/Engine/Plugins/Runtime/GeoReferencing/Source/GeoReferencing/Private/RoundPlanetPawn.cpp
@@ -0,0 +1,435 @@
+// Copyright Epic Games, Inc. All Rights Reserved.
+
+
+#include "RoundPlanetPawn.h"
+#include "GameFramework/FloatingPawnMovement.h"
+#include "GeoReferencingSystem.h"
+#include "GameFramework/PlayerInput.h"
+
+
+ARoundPlanetPawn::ARoundPlanetPawn() : ADefaultPawn()
+{
+	PrimaryActorTick.bCanEverTick = true;
+}
+
+void ARoundPlanetPawn::FlyToLocationECEF(const FVector& ECEFDestination, double YawAtDestination, double PitchAtDestination, bool CanInterruptByMoving)
+{
+	if (!GeoReferencingSystem)
+	{
+		return;
+	}
+	FGeographicCoordinates TempGeographic;
+	FVector TempLocation;
+
+	// We will work on position in ECEF space, because origin is at planet center, that makes some quaternion-based maths easier. 
+	
+	// Compute Origin location in ECEF space
+	FVector ECEFOrigin;
+	GeoReferencingSystem->EngineToECEF(GetActorLocation(), ECEFOrigin);
+
+	// Compute the source and destination rotations in ENU
+	// As during the flight, we can go around the globe, this is better to interpolate in ENU coordinates
+	FlyToLocationOriginRotation = ADefaultPawn::GetViewRotation();
+	FlyToLocationDestinationRotation = FRotator(PitchAtDestination, YawAtDestination, 0);
+
+	// Basically, the motion from Origin to Destination is a the shorted rotation around an axis
+	// This is what a quaternion allow to do easily ! 
+	// Compute Rotation Axis/Angle transform and initialize key points 
+	FQuat TrajectoryQuat = FQuat::FindBetweenVectors(ECEFOrigin, ECEFDestination); // We can do that because ECEF Origin is Planet center 
+	double TrajectoryTotalAngleRadians;
+	FVector TrajectoryRotationAxis;
+	TrajectoryQuat.ToAxisAndAngle(TrajectoryRotationAxis, TrajectoryTotalAngleRadians);
+
+	// Avoid heavy computations during the flight. Pre-generate a fixed list of keypoint, to be interpolated during the motion
+	// Avoid a too small amount of keypoints. 
+	int TrajectorySteps = FMath::Max(int(TrajectoryTotalAngleRadians / FMath::DegreesToRadians(GranularityDegrees)) - 1, 0);
+	TrajectorySteps = FMath::Clamp(TrajectorySteps, MinimumStepCount, MaximumStepCount);
+	FlyToLocationKeypoints.Empty(TrajectorySteps + 2); 
+
+	// We will not create a curve projected along the ellipsoid as we want to get altitude over terrain while flying. 
+	// The radius of the current point will evolve as follow
+	//  - Project the point on the ellipsoid - Will give a default radius depending on ellipsoid location. 
+	//  - Interpolate the altitudes : get origin/destination altitude, and make a linear interpolation between them. This will allow for flying from/to any point smoothly. 
+	//  - Add as flightProfile offset /-\ defined by a curve. 
+
+	// Compute global radius at source and destination points
+	FVector OriginUpECEF, DestinationUpECEF;
+	float OriginRadius, DestinationRadius;
+	ECEFOrigin.ToDirectionAndLength(OriginUpECEF, OriginRadius);
+	ECEFDestination.ToDirectionAndLength(DestinationUpECEF, DestinationRadius);
+
+	// Compute Altitude at origin and destination points by using geographic coordinates
+	float OriginAltitude, DestinationAltitude = 0;
+	GeoReferencingSystem->ECEFToGeographic(ECEFOrigin, TempGeographic);
+	OriginAltitude = TempGeographic.Altitude;
+	GeoReferencingSystem->ECEFToGeographic(ECEFDestination, TempGeographic);
+	DestinationAltitude = TempGeographic.Altitude;
+	
+	// Get distance between source and destination points to compute a wanted altitude from curve
+	// This is not a distance on the ellipsoid just a direct distance between two points to have a ROM of the travel length
+	float TravelDirectDistance = FVector::Dist(ECEFOrigin, ECEFDestination);
+
+	// Add first keypoint
+	GeoReferencingSystem->ECEFToEngine(ECEFOrigin, TempLocation);
+	FlyToLocationKeypoints.Add(TempLocation);
+	
+	// Debug
+	//DrawDebugPoint(GetWorld(), TempLocation, 8, FColor::Red, true, 30);
+
+	for (int CurrentStep = 1; CurrentStep <= TrajectorySteps; CurrentStep++)
+	{
+		float TravelPercentage = (float)CurrentStep / (TrajectorySteps + 1);
+
+		// Lerp over Origin and Destination altitude to account for Altitude changes 
+		float ExtrapolatedAltitude = FMath::Lerp<float>(OriginAltitude, DestinationAltitude, TravelPercentage);
+		float TravelAngle = static_cast<float>(CurrentStep) * FMath::RadiansToDegrees(TrajectoryTotalAngleRadians) / ( TrajectorySteps + 1);
+
+		// Compute the current point by rotating by an angle step
+		FVector RotatedOrigin = ECEFOrigin.RotateAngleAxis(TravelAngle, TrajectoryRotationAxis);
+		
+		// Project on the Ellipsoid / Neutralize Altitude
+		GeoReferencingSystem->ECEFToGeographic(RotatedOrigin, TempGeographic);
+		TempGeographic.Altitude = 0;
+		FVector RotatedProjectedOrigin;
+		GeoReferencingSystem->GeographicToECEF(TempGeographic, RotatedProjectedOrigin);
+
+		// Get Local Up vector
+		FVector LocalUp = RotatedOrigin - RotatedProjectedOrigin;
+		LocalUp.Normalize();
+
+		// Add an altitude if we have a profile curve for it
+		float AltitudeOffset = 0;
+		if (AltitudeProfileCurve != NULL)
+		{
+			float ReferenceAltitude = MaximumAltitudeValue;
+			if (MaximumAltitudeCurve != NULL)
+			{
+				ReferenceAltitude = MaximumAltitudeCurve->GetFloatValue(TravelDirectDistance);
+			}
+			AltitudeOffset = ReferenceAltitude * AltitudeProfileCurve->GetFloatValue(TravelPercentage);
+		}
+
+		// Compute the final keypoint
+		FVector FinalPoint = RotatedProjectedOrigin + LocalUp * (ExtrapolatedAltitude + AltitudeOffset);
+		GeoReferencingSystem->ECEFToEngine(FinalPoint, TempLocation); 
+		FlyToLocationKeypoints.Add(TempLocation);
+
+		// Debug
+		//DrawDebugPoint(GetWorld(), TempLocation, 8, FColor::Purple, true, 30);
+	}
+	GeoReferencingSystem->ECEFToEngine(ECEFDestination, TempLocation);
+	FlyToLocationKeypoints.Add(TempLocation);
+	
+	// Debug
+	//DrawDebugPoint(GetWorld(), TempLocation, 8, FColor::Purple, true, 30);
+
+	// Tell the tick we will be flying from now
+	bIsFlyingToLocation = true;
+	bCanInterruptFlight = CanInterruptByMoving;
+	CurrentFlyTime = 0;
+}
+
+void ARoundPlanetPawn::FlyToLocationGeographic(const FGeographicCoordinates& GeographicDestination, double YawAtDestination, double PitchAtDestination, bool CanInterruptByMoving)
+{
+	if (GeoReferencingSystem)
+	{
+		FVector ECEFDestination;
+		GeoReferencingSystem->GeographicToECEF(GeographicDestination, ECEFDestination);
+		FlyToLocationECEF(ECEFDestination, YawAtDestination, PitchAtDestination, CanInterruptByMoving);
+	}
+
+}
+
+void ARoundPlanetPawn::FlyToLocationLatitudeLongitudeAltitude(const double& InLatitude, const double& InLongitude, const double& InAltitude, double YawAtDestination, double PitchAtDestination, bool CanInterruptByMoving)
+{
+	FlyToLocationGeographic(FGeographicCoordinates(InLongitude, InLatitude, InAltitude), YawAtDestination, PitchAtDestination, CanInterruptByMoving);
+}
+
+void ARoundPlanetPawn::InterruptFlyToLocation()
+{
+	this->bIsFlyingToLocation = false;
+
+	// fix camera roll to 0.0
+	FRotator currentRotator = GetController()->GetControlRotation();
+	currentRotator.Roll = 0.0;
+	GetController()->SetControlRotation(currentRotator);
+}
+
+// Public Overridables
+
+void ARoundPlanetPawn::OnConstruction(const FTransform& Transform)
+{
+	Super::OnConstruction(Transform);
+	if (!GeoReferencingSystem)
+	{
+		GeoReferencingSystem = AGeoReferencingSystem::GetGeoReferencingSystem(GetWorld());
+	}
+}
+
+void ARoundPlanetPawn::BeginPlay()
+{
+	Super::BeginPlay();
+
+	// Round planet is huge, if we enable world bounds check, the pawn behavior might be unpredictable when flying...
+	// LWC_TODO - should not be needed after WORLD_MAX updates. 
+	GetWorldSettings()->bEnableWorldBoundsChecks = false;
+
+	if (!GeoReferencingSystem)
+	{
+		GeoReferencingSystem = AGeoReferencingSystem::GetGeoReferencingSystem(GetWorld());
+
+		if (!GeoReferencingSystem)
+		{
+			UE_LOG(LogGeoReferencing, Error, TEXT("Impossible to use a RoundPlanetPawn without a GeoReferencingSystem."));
+		}
+	}
+
+	// Update cache at first run
+	UpdateMotionObjectsCache(true);
+}
+
+void ARoundPlanetPawn::Tick(float DeltaSeconds)
+{
+	Super::Tick(DeltaSeconds);
+	
+	UpdateMotionObjectsCache();
+	HandleMotionSpeed();
+	ProcessFlyToLocation(DeltaSeconds);
+}
+
+void ARoundPlanetPawn::MoveForward(float Val)
+{
+	if (Val != 0.f)
+	{
+		if (Controller)
+		{
+			FRotator const ControlSpaceRot = this->GetViewRotation();
+
+			FVector MovementVector = FRotationMatrix(ControlSpaceRot).GetScaledAxis(EAxis::X);
+			if (OrbitalMotion)
+			{
+				// In orbital motion, project the motion vector to East/North plane
+				FVector ENUMovementVector = TangentTransform.InverseTransformVector(MovementVector);
+				ENUMovementVector.Z = 0;
+				ENUMovementVector.Normalize();
+				MovementVector = TangentTransform.TransformVector(ENUMovementVector);
+			}
+			// transform to world space and add it
+			AddMovementInput(MovementVector, Val);
+
+			// Allows for recovering user control
+			if (bIsFlyingToLocation && bCanInterruptFlight)
+			{
+				InterruptFlyToLocation();
+			}
+		}
+	}
+}
+
+void ARoundPlanetPawn::MoveRight(float Val)
+{
+	if (Val != 0.f)
+	{
+		if (Controller)
+		{
+			FRotator const ControlSpaceRot = this->GetViewRotation();
+
+			FVector MovementVector = FRotationMatrix(ControlSpaceRot).GetScaledAxis(EAxis::Y);
+			if (OrbitalMotion)
+			{
+				// In orbital motion, project the motion vector to East/North plane
+				FVector ENUMovementVector = TangentTransform.InverseTransformVector(MovementVector);
+				ENUMovementVector.Z = 0;
+				ENUMovementVector.Normalize();
+				MovementVector = TangentTransform.TransformVector(ENUMovementVector);
+			}
+			// transform to world space and add it
+			AddMovementInput(MovementVector, Val);
+
+			// Allows for recovering user control
+			if (bIsFlyingToLocation && bCanInterruptFlight)
+			{
+				InterruptFlyToLocation();
+			}
+		}
+	}
+}
+
+void ARoundPlanetPawn::MoveUp_World(float Val)
+{
+	if (Val != 0.f)
+	{
+		// Don't use UE Z here, but the Planet local Up
+		AddMovementInput(Up, Val);
+
+		if (bIsFlyingToLocation && bCanInterruptFlight)
+		{
+			InterruptFlyToLocation();
+		}
+	}
+}
+
+FRotator ARoundPlanetPawn::GetViewRotation() const
+{
+	// Compensate for tangent motion orientation
+	FRotator LocalRotator = ADefaultPawn::GetViewRotation();
+	return TangentTransform.TransformRotation(LocalRotator.Quaternion()).Rotator();
+}
+
+FRotator ARoundPlanetPawn::GetBaseAimRotation() const 
+{
+	return this->GetViewRotation();
+}
+
+void ARoundPlanetPawn::SetupPlayerInputComponent(UInputComponent* InInputComponent)
+{
+	Super::SetupPlayerInputComponent(InInputComponent);
+	
+	// Add MouseScroll-based speed change behavior, similarely as done in the Editor.
+	check(InInputComponent);
+	if (bAddDefaultMovementBindings)
+	{
+		static bool bRoundPlanetBindingsAdded = false;
+		if (!bRoundPlanetBindingsAdded)
+		{
+			bRoundPlanetBindingsAdded = true;
+
+			UPlayerInput::AddEngineDefinedActionMapping(FInputActionKeyMapping("RoundPlanetPawn_IncreaseSpeedScalar", EKeys::MouseScrollUp));
+			UPlayerInput::AddEngineDefinedActionMapping(FInputActionKeyMapping("RoundPlanetPawn_DecreaseSpeedScalar", EKeys::MouseScrollDown));
+			UPlayerInput::AddEngineDefinedActionMapping(FInputActionKeyMapping("RoundPlanetPawn_ResetSpeedScalar", EKeys::MiddleMouseButton));
+		}
+
+		InInputComponent->BindAction("RoundPlanetPawn_IncreaseSpeedScalar", IE_Pressed, this, &ARoundPlanetPawn::IncreaseSpeedScalar);
+		InInputComponent->BindAction("RoundPlanetPawn_DecreaseSpeedScalar", IE_Pressed, this, &ARoundPlanetPawn::DecreaseSpeedScalar);
+		InInputComponent->BindAction("RoundPlanetPawn_ResetSpeedScalar", IE_Pressed, this, &ARoundPlanetPawn::ResetSpeedScalar);
+	}
+}
+
+// Private Methods
+void ARoundPlanetPawn::ProcessFlyToLocation(float DeltaSeconds)
+{
+	if (bIsFlyingToLocation)
+	{
+		CurrentFlyTime += DeltaSeconds;
+		float TimeProgress = CurrentFlyTime / FlyDuration;
+
+		if (TimeProgress < 1.0)
+		{
+			// In order to accelerate at start and slow down at end, we use a progress profile curve
+			float FlightProgress = TimeProgress;
+			if (ProgressCurve != NULL)
+			{
+				FlightProgress = FMath::Clamp<float>(ProgressCurve->GetFloatValue(TimeProgress), 0.0, 0.9999);
+			}
+
+			// Find the bounding keypoint indexes corresponding to the current Flight progress
+			int PreviousKeypointIndex = FMath::Floor(FlightProgress * (FlyToLocationKeypoints.Num() - 1));
+			int NextKeypointIndex = PreviousKeypointIndex + 1;
+			float SegmentProgress = FlightProgress * (FlyToLocationKeypoints.Num() - 1) - PreviousKeypointIndex;
+
+			// Get the current position by interpolating between those two points
+			FVector PreviousPosition = FlyToLocationKeypoints[PreviousKeypointIndex];
+			FVector NextPosition = FlyToLocationKeypoints[NextKeypointIndex];
+			FVector CurrentPosition = FMath::Lerp<FVector>(PreviousPosition, NextPosition, SegmentProgress);
+			
+			// Set Location
+			SetActorLocation(CurrentPosition); 
+
+			FTransform CurrentTangentTransform = GeoReferencingSystem->GetTangentTransformAtEngineLocation(CurrentPosition);
+			FQuat OriginRotationENU = CurrentTangentTransform.InverseTransformRotation(FlyToLocationOriginRotation.Quaternion()); 
+			FQuat DestinationRotationENU = CurrentTangentTransform.InverseTransformRotation(FlyToLocationDestinationRotation.Quaternion());
+			
+			// Interpolate rotation - ENU space, then UE space. If we don't do like that, we can get upside down. 
+			FQuat CurrentRotatorENU = FMath::Lerp<FQuat>(OriginRotationENU, DestinationRotationENU, FlightProgress);
+			FRotator CurrentRotator = CurrentTangentTransform.TransformRotation(CurrentRotatorENU).Rotator();
+			GetController()->SetControlRotation(CurrentRotator);
+		}
+		else
+		{
+			// We reached the end - Set actual destination location 
+			FVector FinalPoint = FlyToLocationKeypoints.Last();
+			SetActorLocation(FinalPoint);
+			GetController()->SetControlRotation(FlyToLocationDestinationRotation);
+			MovementComponent->StopMovementImmediately();
+			bIsFlyingToLocation = false;
+			CurrentFlyTime = 0;
+		}
+	}
+}
+
+void ARoundPlanetPawn::HandleMotionSpeed()
+{
+	// Handle movement speed
+	float MaxSpeed = BaseSpeedKmh * SpeedScalar * 100.0 / 3.6;  // Km/h to UE units/s
+	if (SpeedByHATModifierCurve)
+	{
+		MaxSpeed *= SpeedByHATModifierCurve->GetFloatValue(HAT);
+	}
+
+	// We need to update the Acceleration/Deceleration accordingly and limit max speed to have a similar behavior whatever the speed scalar. 
+	UFloatingPawnMovement* FPMovementComponent = Cast<UFloatingPawnMovement>(MovementComponent);
+	{
+		FPMovementComponent->MaxSpeed = MaxSpeed;
+		FPMovementComponent->Acceleration = MaxSpeed * 4000.f / 1200.f; // Keep the coefficient proportional to the default values. 
+		FPMovementComponent->Deceleration = MaxSpeed * 8000.f / 1200.f; // Keep the coefficient proportional to the default values. 
+		
+		if (FPMovementComponent->IsExceedingMaxSpeed(MaxSpeed))
+		{
+			FVector Velocity = FPMovementComponent->Velocity;
+			Velocity.Normalize();
+			FPMovementComponent->Velocity = Velocity * MaxSpeed * 0.999;
+		}
+	}
+}
+
+void ARoundPlanetPawn::UpdateMotionObjectsCache(bool ForceUpdate /*false*/)
+{
+	if (!GeoReferencingSystem)
+	{
+		return;
+	}
+
+	// Update Local ENU Transform
+	if ((this->GetActorLocation() - LastCacheLocation).SquaredLength() > CacheUpdateThresholdSq ||
+		ForceUpdate)
+	{
+		// Update ENU Vectors
+		GeoReferencingSystem->GetENUVectorsAtEngineLocation(GetActorLocation(), East, North, Up);
+		FGeographicCoordinates GeographicCoordinates;
+		GeoReferencingSystem->EngineToGeographic(GetActorLocation(), GeographicCoordinates);
+
+		// Update Tangent Transform
+		TangentTransform = GeoReferencingSystem->GetTangentTransformAtEngineLocation(GetActorLocation());
+
+		// Update Altitude
+		Altitude = static_cast<float>(GeographicCoordinates.Altitude);
+
+		// Schedule a new HAT update Query
+		bIsHATInvalid = true; 
+	}
+
+	if (bIsHATInvalid)
+	{
+		// Update HAT each step if not found previously, or scheduled by the threshold exceeded. 
+		FVector LineCheckStart = GetActorLocation();
+		FVector LineCheckEnd = GetActorLocation() - (Altitude + 0.05 * GeoReferencingSystem->GetGeographicEllipsoidMaxRadius()) * Up; // Estimate raycast length - Altitude + 5% of ellipsoid in case of negative altitudes
+		FHitResult HitResult = FHitResult();
+
+		static const FName LineTraceSingleName(TEXT("RoundPlanetPawnLineTrace"));
+		FCollisionQueryParams CollisionParams(LineTraceSingleName);
+		CollisionParams.bTraceComplex = true;
+		CollisionParams.AddIgnoredActor(this);
+
+		FCollisionObjectQueryParams ObjectParams = FCollisionObjectQueryParams(ECC_WorldStatic);
+		ObjectParams.AddObjectTypesToQuery(ECC_WorldDynamic);
+		ObjectParams.AddObjectTypesToQuery(ECC_Pawn);
+		ObjectParams.AddObjectTypesToQuery(ECC_Visibility);
+
+		
+		if (GetWorld()->LineTraceSingleByObjectType(HitResult, LineCheckStart, LineCheckEnd, ObjectParams, CollisionParams))
+		{
+			HAT = FVector::Dist(GetActorLocation(), HitResult.ImpactPoint) / 100.0;
+			bIsHATInvalid = false;
+		}
+	}
+}
diff --git a/Engine/Plugins/Runtime/GeoReferencing/Source/GeoReferencing/Public/CartesianCoordinates.h b/Engine/Plugins/Runtime/GeoReferencing/Source/GeoReferencing/Public/CartesianCoordinates.h
--- a/Engine/Plugins/Runtime/GeoReferencing/Source/GeoReferencing/Public/CartesianCoordinates.h
+++ b/Engine/Plugins/Runtime/GeoReferencing/Source/GeoReferencing/Public/CartesianCoordinates.h
@@ -7,7 +7,7 @@
 #include "CartesianCoordinates.generated.h"
 
 
-USTRUCT(BlueprintType)
+USTRUCT(BlueprintType, meta = (Deprecated = "5.0"))
 
 struct GEOREFERENCING_API FCartesianCoordinates
 {
@@ -22,21 +22,19 @@ public:
 	double Y;
 	double Z;
 
+	
 	FText ToFullText(int32 IntegralDigits = 3);
 	FText ToCompactText(int32 IntegralDigits = 3);
 	void ToSeparateTexts(FText& OutX, FText& OutY, FText& OutZ, int32 IntegralDigits = 3);
-
 	void ToFloatApproximation(float& OutX, float& OutY, float& OutZ);
-
 	FVector ToVector() const;
-	
-	PRAGMA_DISABLE_DEPRECATION_WARNINGS // Suppress compiler warning on override of deprecated function
 
+	PRAGMA_DISABLE_DEPRECATION_WARNINGS // Suppress compiler warning on override of deprecated function
 
-		UE_DEPRECATED(5.0, "Use FCartesianCoordinates::ToVector() instead.")
-		FVector3d ToVector3d() const;
-		UE_DEPRECATED(5.0, "Use FCartesianCoordinates(const FVector& Coordinates) instead.")
-		FCartesianCoordinates(const FVector4d& Coordinates);
+	UE_DEPRECATED(5.0, "Use FCartesianCoordinates::ToVector() instead.")
+	FVector3d ToVector3d() const;
+	UE_DEPRECATED(5.0, "Use FCartesianCoordinates(const FVector& Coordinates) instead.")
+	FCartesianCoordinates(const FVector4d& Coordinates);
 
 	PRAGMA_ENABLE_DEPRECATION_WARNINGS
 
@@ -51,7 +49,8 @@ class GEOREFERENCING_API UCartesianCoordinatesFunctionLibrary : public UBlueprin
 	/**
 	 * Converts a LargeCoordinates value to localized formatted text, in the form 'X= Y= Z='
 	 **/
-	UFUNCTION(BlueprintPure, meta = (DisplayName = "ToFullText", AdvancedDisplay = "1", BlueprintAutocast), Category = "GeoReferencing")
+	UE_DEPRECATED(5.0, "FCartesianCoordinates is deprecated : Use the version that uses a FVector instead.")
+	UFUNCTION(BlueprintPure, meta = (DisplayName = "ToFullText", AdvancedDisplay = "1", BlueprintAutocast, DeprecatedFunction, DeprecationMessage = "FCartesianCoordinates is deprecated : Use the version that uses a FVector instead."), Category = "GeoReferencing")
 	static FORCEINLINE FText ToFullText(UPARAM(ref) FCartesianCoordinates& CartesianCoordinates, int32 IntegralDigits = 3)
 	{
 		return CartesianCoordinates.ToFullText(IntegralDigits);
@@ -60,7 +59,8 @@ class GEOREFERENCING_API UCartesianCoordinatesFunctionLibrary : public UBlueprin
 	/**
 	 * Converts a LargeCoordinates value to formatted text, in the form '(X, Y, Z)'
 	 **/
-	UFUNCTION(BlueprintPure, meta = (DisplayName = "ToCompactText", AdvancedDisplay = "1", BlueprintAutocast), Category = "GeoReferencing")
+	UE_DEPRECATED(5.0, "FCartesianCoordinates is deprecated : Use the version that uses a FVector instead.")
+	UFUNCTION(BlueprintPure, meta = (DisplayName = "ToCompactText", AdvancedDisplay = "1", BlueprintAutocast, DeprecatedFunction, DeprecationMessage = "FCartesianCoordinates is deprecated : Use the version that uses a FVector instead."), Category = "GeoReferencing")
 	static FORCEINLINE FText ToCompactText(UPARAM(ref) FCartesianCoordinates& CartesianCoordinates, int32 IntegralDigits = 3)
 	{
 		return CartesianCoordinates.ToCompactText(IntegralDigits);
@@ -69,7 +69,8 @@ class GEOREFERENCING_API UCartesianCoordinatesFunctionLibrary : public UBlueprin
 	/**
 	 * Converts a LargeCoordinates value to 3 separate text values
 	 **/
-	UFUNCTION(BlueprintPure, meta = (DisplayName = "ToSeparateTexts", AdvancedDisplay = "4", BlueprintAutocast), Category = "GeoReferencing")
+	UE_DEPRECATED(5.0, "FCartesianCoordinates is deprecated : Use the version that uses a FVector instead.")
+	UFUNCTION(BlueprintPure, meta = (DisplayName = "ToSeparateTexts", AdvancedDisplay = "4", BlueprintAutocast, DeprecatedFunction, DeprecationMessage = "FCartesianCoordinates is deprecated : Use the version that uses a FVector instead."), Category = "GeoReferencing")
 	static void ToSeparateTexts(UPARAM(ref) FCartesianCoordinates& CartesianCoordinates, FText& OutX, FText& OutY, FText& OutZ, int32 IntegralDigits = 3)
 	{
 		CartesianCoordinates.ToSeparateTexts(OutX, OutY, OutZ, IntegralDigits);
@@ -79,7 +80,8 @@ class GEOREFERENCING_API UCartesianCoordinatesFunctionLibrary : public UBlueprin
 	 * Get the Coordinates as a float approximation.
 	 * USE WISELY as we can't guarantee there will no be rounding due to IEEE754 float encoding !
 	 **/
-	UFUNCTION(BlueprintPure, Category = "GeoReferencing")
+	UE_DEPRECATED(5.0, "FCartesianCoordinates is deprecated : Use the version that uses a FVector instead.")
+	UFUNCTION(BlueprintPure, Category = "GeoReferencing", meta = (DeprecatedFunction, DeprecationMessage = "FCartesianCoordinates is deprecated : Use the version that uses a FVector instead."))
 	static void ToFloatApproximation(UPARAM(ref) FCartesianCoordinates& CartesianCoordinates, float& OutX, float& OutY, float& OutZ)
 	{
 		CartesianCoordinates.ToFloatApproximation(OutX, OutY, OutZ);
@@ -89,7 +91,8 @@ class GEOREFERENCING_API UCartesianCoordinatesFunctionLibrary : public UBlueprin
 	 * Set the Coordinates from float approximation.
 	 * USE WISELY as we can't guarantee there will no be rounding due to IEEE754 float encoding !
 	 **/
-	UFUNCTION(BlueprintPure, Category = "GeoReferencing")
+	UE_DEPRECATED(5.0, "FCartesianCoordinates is deprecated : Use the version that uses a FVector instead.")
+	UFUNCTION(BlueprintPure, Category = "GeoReferencing", meta = (DeprecatedFunction, DeprecationMessage = "FCartesianCoordinates is deprecated : Use the version that uses a FVector instead."))
 	static FCartesianCoordinates MakeCartesianCoordinatesApproximation(const float& InX, const float& InY, const float& InZ)
 	{
 		return FCartesianCoordinates(static_cast<float>(InX), static_cast<float>(InY), static_cast<float>(InZ));
diff --git a/Engine/Plugins/Runtime/GeoReferencing/Source/GeoReferencing/Public/Ellipsoid.h b/Engine/Plugins/Runtime/GeoReferencing/Source/GeoReferencing/Public/Ellipsoid.h
--- a/Engine/Plugins/Runtime/GeoReferencing/Source/GeoReferencing/Public/Ellipsoid.h
+++ b/Engine/Plugins/Runtime/GeoReferencing/Source/GeoReferencing/Public/Ellipsoid.h
@@ -24,6 +24,9 @@ public:
 	FVector OneOverRadii;
 	FVector OneOverRadiiSquared;
 
-	FVector GeodeticSurfaceNormal(const FCartesianCoordinates& ECEFLocation) const;
-	FVector GeodeticSurfaceNormal(const FGeographicCoordinates& GeographicCoordinates) const;    
+	double GetMaximumRadius();
+	double GetMinimumRadius();
+
+	FVector GeodeticSurfaceNormal(const FVector& ECEFLocation) const;
+	FVector GeodeticSurfaceNormal(const FGeographicCoordinates& GeographicCoordinates) const;
 };
diff --git a/Engine/Plugins/Runtime/GeoReferencing/Source/GeoReferencing/Public/GeoReferencingBFL.h b/Engine/Plugins/Runtime/GeoReferencing/Source/GeoReferencing/Public/GeoReferencingBFL.h
new file mode 100644
--- /dev/null
+++ b/Engine/Plugins/Runtime/GeoReferencing/Source/GeoReferencing/Public/GeoReferencingBFL.h
@@ -0,0 +1,35 @@
+// Copyright Epic Games, Inc. All Rights Reserved.
+
+#pragma once
+
+#include "CoreMinimal.h"
+#include "Kismet/BlueprintFunctionLibrary.h"
+#include "GeoReferencingBFL.generated.h"
+
+/**
+ * Blueprint function library to convert geospatial coordinates to text
+ */
+UCLASS()
+class GEOREFERENCING_API UGeoReferencingBFL : public UBlueprintFunctionLibrary
+{
+	GENERATED_BODY()
+
+	/**
+	* Converts a LargeCoordinates value to localized formatted text, in the form 'X= Y= Z='
+	**/
+	UFUNCTION(BlueprintPure, Category = "GeoReferencing", meta = (DisplayName = "ToFullText", AdvancedDisplay = "1", BlueprintAutocast))
+	static FText ToFullText(UPARAM(ref) FVector& CartesianCoordinates, int32 IntegralDigits = 3);
+
+	/**
+	 * Converts a LargeCoordinates value to formatted text, in the form '(X, Y, Z)'
+	 **/
+	UFUNCTION(BlueprintPure, Category = "GeoReferencing", meta = (DisplayName = "ToCompactText", AdvancedDisplay = "1", BlueprintAutocast))
+	static FText ToCompactText(UPARAM(ref) FVector& CartesianCoordinates, int32 IntegralDigits = 3);
+
+	/**
+	 * Converts a LargeCoordinates value to 3 separate text values
+	 **/
+	UFUNCTION(BlueprintPure, Category = "GeoReferencing", meta = (DisplayName = "ToSeparateTexts", AdvancedDisplay = "4", BlueprintAutocast))
+	static void ToSeparateTexts(UPARAM(ref) FVector& CartesianCoordinates, FText& OutX, FText& OutY, FText& OutZ, int32 IntegralDigits = 3);
+
+};
diff --git a/Engine/Plugins/Runtime/GeoReferencing/Source/GeoReferencing/Public/GeoReferencingModule.h b/Engine/Plugins/Runtime/GeoReferencing/Source/GeoReferencing/Public/GeoReferencingModule.h
new file mode 100644
--- /dev/null
+++ b/Engine/Plugins/Runtime/GeoReferencing/Source/GeoReferencing/Public/GeoReferencingModule.h
@@ -0,0 +1,12 @@
+// Copyright Epic Games, Inc. All Rights Reserved.
+
+#pragma once
+
+#include "CoreMinimal.h"
+#include "Modules/ModuleInterface.h"
+
+GEOREFERENCING_API DECLARE_LOG_CATEGORY_EXTERN(LogGeoReferencing, Log, All);
+
+class GEOREFERENCING_API FGeoReferencingModule : public IModuleInterface
+{
+};
diff --git a/Engine/Plugins/Runtime/GeoReferencing/Source/GeoReferencing/Public/GeoReferencingSystem.h b/Engine/Plugins/Runtime/GeoReferencing/Source/GeoReferencing/Public/GeoReferencingSystem.h
--- a/Engine/Plugins/Runtime/GeoReferencing/Source/GeoReferencing/Public/GeoReferencingSystem.h
+++ b/Engine/Plugins/Runtime/GeoReferencing/Source/GeoReferencing/Public/GeoReferencingSystem.h
@@ -14,7 +14,7 @@
 
 #include "GeoReferencingSystem.generated.h"
 
-DECLARE_LOG_CATEGORY_EXTERN(LogGeoReferencing, Log, All);
+
 
 UENUM(BlueprintType)
 enum class EPlanetShape : uint8 {
@@ -33,9 +33,10 @@ enum class EPlanetShape : uint8 {
 };
 
 /**
- *
+ * This AInfos enable you to define a correspondance between the UE origin and an actual geographic location on a planet
+ * Once done it offers different functions to convert coordinates between UE and Geographic coordinates
  */
-UCLASS(hidecategories=(Transform, Replication, Actor, Cooking))
+UCLASS(hidecategories = (Transform, Replication, Actor, Cooking))
 class GEOREFERENCING_API AGeoReferencingSystem : public AInfo
 {
 	GENERATED_BODY()
@@ -49,168 +50,385 @@ public:
 	UFUNCTION(BlueprintPure, Category = "GeoReferencing", meta = (WorldContext = "WorldContextObject"))
 	static AGeoReferencingSystem* GetGeoReferencingSystem(UObject* WorldContextObject);
 
-	/**
-	  * Convert a Vector expressed in ENGINE space to the PROJECTED CRS
-	 */
-	UFUNCTION(BlueprintCallable, Category = "GeoReferencing|Transformations")
+	// Deprecate all previous functions that were relying on Double precision placeholders
+
+#pragma region Old deprecated Prototypes
+
+	PRAGMA_DISABLE_DEPRECATION_WARNINGS
+
+	UE_DEPRECATED(5.0, "FCartesianCoordinates is deprecated : Use the version that uses a FVector instead.")
+	UFUNCTION(BlueprintCallable, Category = "GeoReferencing|Transformations", meta = (DeprecatedFunction, DeprecationMessage = "FCartesianCoordinates is deprecated : Use the version that uses a FVector instead."))
 	void EngineToProjected(const FVector& EngineCoordinates, FCartesianCoordinates& ProjectedCoordinates);
 
-	/**
-	 * Convert a Vector expressed in PROJECTED CRS to ENGINE space
-	 */
-	UFUNCTION(BlueprintCallable, Category = "GeoReferencing|Transformations")
+	UE_DEPRECATED(5.0, "FCartesianCoordinates is deprecated : Use the version that uses a FVector instead.")
+	UFUNCTION(BlueprintCallable, Category = "GeoReferencing|Transformations", meta = (DeprecatedFunction, DeprecationMessage = "FCartesianCoordinates is deprecated : Use the version that uses a FVector instead."))
 	void ProjectedToEngine(const FCartesianCoordinates& ProjectedCoordinates, FVector& EngineCoordinates);
 
+	UE_DEPRECATED(5.0, "FCartesianCoordinates is deprecated : Use the version that uses a FVector instead.")
+	UFUNCTION(BlueprintCallable, Category = "GeoReferencing|Transformations", meta = (DeprecatedFunction, DeprecationMessage = "FCartesianCoordinates is deprecated : Use the version that uses a FVector instead."))
+	void EngineToECEF(const FVector& EngineCoordinates, FCartesianCoordinates& ECEFCoordinates);
+
+	UE_DEPRECATED(5.0, "FCartesianCoordinates is deprecated : Use the version that uses a FVector instead.")
+	UFUNCTION(BlueprintCallable, Category = "GeoReferencing|Transformations", meta = (DeprecatedFunction, DeprecationMessage = "FCartesianCoordinates is deprecated : Use the version that uses a FVector instead."))
+	void ECEFToEngine(const FCartesianCoordinates& ECEFCoordinates, FVector& EngineCoordinates);
+
+	UE_DEPRECATED(5.0, "FCartesianCoordinates is deprecated : Use the version that uses a FVector instead.")
+	UFUNCTION(BlueprintCallable, Category = "GeoReferencing|Transformations", meta = (DeprecatedFunction, DeprecationMessage = "FCartesianCoordinates is deprecated : Use the version that uses a FVector instead."))
+	void ProjectedToGeographic(const FCartesianCoordinates& ProjectedCoordinates, FGeographicCoordinates& GeographicCoordinates);
+
+	UE_DEPRECATED(5.0, "FCartesianCoordinates is deprecated : Use the version that uses a FVector instead.")
+	UFUNCTION(BlueprintCallable, Category = "GeoReferencing|Transformations", meta = (DeprecatedFunction, DeprecationMessage = "FCartesianCoordinates is deprecated : Use the version that uses a FVector instead."))
+	void GeographicToProjected(const FGeographicCoordinates& GeographicCoordinates, FCartesianCoordinates& ProjectedCoordinates);
+
+	UE_DEPRECATED(5.0, "FCartesianCoordinates is deprecated : Use the version that uses a FVector instead.")
+	UFUNCTION(BlueprintCallable, Category = "GeoReferencing|Transformations", meta = (DeprecatedFunction, DeprecationMessage = "FCartesianCoordinates is deprecated : Use the version that uses a FVector instead."))
+	void ProjectedToECEF(const FCartesianCoordinates& ProjectedCoordinates, FCartesianCoordinates& ECEFCoordinates);
+
+	UE_DEPRECATED(5.0, "FCartesianCoordinates is deprecated : Use the version that uses a FVector instead.")
+	UFUNCTION(BlueprintCallable, Category = "GeoReferencing|Transformations", meta = (DeprecatedFunction, DeprecationMessage = "FCartesianCoordinates is deprecated : Use the version that uses a FVector instead."))
+	void ECEFToProjected(const FCartesianCoordinates& ECEFCoordinates, FCartesianCoordinates& ProjectedCoordinates);
+
+	UE_DEPRECATED(5.0, "FCartesianCoordinates is deprecated : Use the version that uses a FVector instead.")
+	UFUNCTION(BlueprintCallable, Category = "GeoReferencing|Transformations", meta = (DeprecatedFunction, DeprecationMessage = "FCartesianCoordinates is deprecated : Use the version that uses a FVector instead."))
+	void GeographicToECEF(const FGeographicCoordinates& GeographicCoordinates, FCartesianCoordinates& ECEFCoordinates);
+
+	UE_DEPRECATED(5.0, "FCartesianCoordinates is deprecated : Use the version that uses a FVector instead.")
+	UFUNCTION(BlueprintCallable, Category = "GeoReferencing|Transformations", meta = (DeprecatedFunction, DeprecationMessage = "FCartesianCoordinates is deprecated : Use the version that uses a FVector instead."))
+	void ECEFToGeographic(const FCartesianCoordinates& ECEFCoordinates, FGeographicCoordinates& GeographicCoordinates);
+
+	UE_DEPRECATED(5.0, "FCartesianCoordinates is deprecated : Use the version that uses a FVector instead.")
+	UFUNCTION(BlueprintCallable, Category = "GeoReferencing|ENU", meta = (DeprecatedFunction, DeprecationMessage = "FCartesianCoordinates is deprecated : Use the version that uses a FVector instead."))
+	void GetENUVectorsAtProjectedLocation(const FCartesianCoordinates& ProjectedCoordinates, FVector& East, FVector& North, FVector& Up);
+
+	UE_DEPRECATED(5.0, "FCartesianCoordinates is deprecated : Use the version that uses a FVector instead.")
+	UFUNCTION(BlueprintCallable, Category = "GeoReferencing|ENU", meta = (DeprecatedFunction, DeprecationMessage = "FCartesianCoordinates is deprecated : Use the version that uses a FVector instead."))
+	void GetENUVectorsAtECEFLocation(const FCartesianCoordinates& ECEFCoordinates, FVector& East, FVector& North, FVector& Up);
+
+	UE_DEPRECATED(5.0, "FCartesianCoordinates is deprecated : Use the version that uses a FVector instead.")
+	UFUNCTION(BlueprintCallable, Category = "GeoReferencing|ENU", meta = (DeprecatedFunction, DeprecationMessage = "FCartesianCoordinates is deprecated : Use the version that uses a FVector instead."))
+	void GetECEFENUVectorsAtECEFLocation(const FCartesianCoordinates& ECEFCoordinates, FVector& ECEFEast, FVector& ECEFNorth, FVector& ECEFUp);
+
+	UE_DEPRECATED(5.0, "FCartesianCoordinates is deprecated : Use the version that uses a FVector instead.")
+	UFUNCTION(BlueprintCallable, Category = "GeoReferencing|TangentTransforms", meta = (DeprecatedFunction, DeprecationMessage = "FCartesianCoordinates is deprecated : Use the version that uses a FVector instead."))
+	FTransform GetTangentTransformAtProjectedLocation(const FCartesianCoordinates& ProjectedCoordinates);
+
+	UE_DEPRECATED(5.0, "FCartesianCoordinates is deprecated : Use the version that uses a FVector instead.")
+	UFUNCTION(BlueprintCallable, Category = "GeoReferencing|TangentTransforms", meta = (DeprecatedFunction, DeprecationMessage = "FCartesianCoordinates is deprecated : Use the version that uses a FVector instead."))
+	FTransform GetTangentTransformAtECEFLocation(const FCartesianCoordinates& ECEFCoordinates);
+
+	PRAGMA_ENABLE_DEPRECATION_WARNINGS
+
+#pragma endregion
+
+#pragma region New Prototypes for Blueprints
+	// We want to keep the same function names, but with a change in some argument types. UBT doesn't support that unless
+	// we create the new functions with a K2_ prefix and use the meta/Displayname tag to keep the same name. 
+	// To allow for a simple C++ usage too, we keep the C++ functions doing the actual work separately
+
+
+	/**
+	* Convert a Vector expressed in ENGINE space to the PROJECTED CRS
+	*/
+	UFUNCTION(BlueprintCallable, Category = "GeoReferencing|Transformations", meta = (DisplayName = "Engine To Projected"))
+	void K2_EngineToProjected(const FVector& EngineCoordinates, FVector& ProjectedCoordinates) { EngineToProjected(EngineCoordinates, ProjectedCoordinates); }
+
+	/**
+	* Convert a Vector expressed in PROJECTED CRS to ENGINE space
+	*/
+	UFUNCTION(BlueprintCallable, Category = "GeoReferencing|Transformations", meta = (DisplayName = "Projected To Engine"))
+	void K2_ProjectedToEngine(const FVector& ProjectedCoordinates, FVector& EngineCoordinates) { ProjectedToEngine(ProjectedCoordinates, EngineCoordinates); }
+
 	// Engine <--> ECEF 
 
 	/**
-	 * Convert a Vector expressed in ENGINE space to the ECEF CRS
-	 */
-	UFUNCTION(BlueprintCallable, Category = "GeoReferencing|Transformations")
-	void EngineToECEF(const FVector& EngineCoordinates, FCartesianCoordinates& ECEFCoordinates);
+	* Convert a Vector expressed in ENGINE space to the ECEF CRS
+	*/
+	UFUNCTION(BlueprintCallable, Category = "GeoReferencing|Transformations", meta = (DisplayName = "Engine To ECEF"))
+	void K2_EngineToECEF(const FVector& EngineCoordinates, FVector& ECEFCoordinates) { EngineToECEF(EngineCoordinates, ECEFCoordinates); }
 
 	/**
-	 * Convert a Vector expressed in ECEF CRS to ENGINE space
-	 */
-	UFUNCTION(BlueprintCallable, Category = "GeoReferencing|Transformations")
-	void ECEFToEngine(const FCartesianCoordinates& ECEFCoordinates, FVector& EngineCoordinates);
+	* Convert a Vector expressed in ECEF CRS to ENGINE space
+	*/
+	UFUNCTION(BlueprintCallable, Category = "GeoReferencing|Transformations", meta = (DisplayName = "ECEF To Engine"))
+	void K2_ECEFToEngine(const FVector& ECEFCoordinates, FVector& EngineCoordinates) { ECEFToEngine(ECEFCoordinates, EngineCoordinates); }
 
 	// Projected <--> Geographic
 
 	/**
-	 * Convert a Coordinate expressed in PROJECTED CRS to GEOGRAPHIC CRS
-	 */
-	UFUNCTION(BlueprintCallable, Category = "GeoReferencing|Transformations")
-	void ProjectedToGeographic(const FCartesianCoordinates& ProjectedCoordinates, FGeographicCoordinates& GeographicCoordinates);
+	* Convert a Coordinate expressed in PROJECTED CRS to GEOGRAPHIC CRS
+	*/
+	UFUNCTION(BlueprintCallable, Category = "GeoReferencing|Transformations", meta = (DisplayName = "Projected To Geographic"))
+	void K2_ProjectedToGeographic(const FVector& ProjectedCoordinates, FGeographicCoordinates& GeographicCoordinates) { ProjectedToGeographic(ProjectedCoordinates, GeographicCoordinates); }
 
 	/**
-	 * Convert a Coordinate expressed in GEOGRAPHIC CRS to PROJECTED CRS
-	 */
-	UFUNCTION(BlueprintCallable, Category = "GeoReferencing|Transformations")
-	void GeographicToProjected(const FGeographicCoordinates& GeographicCoordinates, FCartesianCoordinates& ProjectedCoordinates);
-
+	* Convert a Coordinate expressed in GEOGRAPHIC CRS to PROJECTED CRS
+	*/
+	UFUNCTION(BlueprintCallable, Category = "GeoReferencing|Transformations", meta = (DisplayName = "Geographic To Projected"))
+	void K2_GeographicToProjected(const FGeographicCoordinates& GeographicCoordinates, FVector& ProjectedCoordinates) { GeographicToProjected(GeographicCoordinates, ProjectedCoordinates); }
 
 	// Projected <--> ECEF
+
 	/**
-	 * Convert a Coordinate expressed in PROJECTED CRS to ECEF CRS
-	 */
-	UFUNCTION(BlueprintCallable, Category = "GeoReferencing|Transformations")
-	void ProjectedToECEF(const FCartesianCoordinates& ProjectedCoordinates, FCartesianCoordinates& ECEFCoordinates);
+	* Convert a Coordinate expressed in PROJECTED CRS to ECEF CRS
+	*/
+	UFUNCTION(BlueprintCallable, Category = "GeoReferencing|Transformations", meta = (DisplayName = "Projected To ECEF"))
+	void K2_ProjectedToECEF(const FVector& ProjectedCoordinates, FVector& ECEFCoordinates) { ProjectedToECEF(ProjectedCoordinates, ECEFCoordinates); }
 
 	/**
-	 * Convert a Coordinate expressed in ECEF CRS to PROJECTED CRS
-	 */
-	UFUNCTION(BlueprintCallable, Category = "GeoReferencing|Transformations")
-	void ECEFToProjected(const FCartesianCoordinates& ECEFCoordinates, FCartesianCoordinates& ProjectedCoordinates);
+	* Convert a Coordinate expressed in ECEF CRS to PROJECTED CRS
+	*/
+	UFUNCTION(BlueprintCallable, Category = "GeoReferencing|Transformations", meta = (DisplayName = "ECEF To Projected"))
+	void K2_ECEFToProjected(const FVector& ECEFCoordinates, FVector& ProjectedCoordinates) { ECEFToProjected(ECEFCoordinates, ProjectedCoordinates); }
 
 	// Geographic <--> ECEF
 
 	/**
-	 * Convert a Coordinate expressed in GEOGRAPHIC CRS to ECEF CRS
-	 */
+	* Convert a Coordinate expressed in GEOGRAPHIC CRS to ECEF CRS
+	*/
+	UFUNCTION(BlueprintCallable, Category = "GeoReferencing|Transformations", meta = (DisplayName = "Geographic To ECEF"))
+	void K2_GeographicToECEF(const FGeographicCoordinates& GeographicCoordinates, FVector& ECEFCoordinates) { GeographicToECEF(GeographicCoordinates, ECEFCoordinates); }
+
+	/**
+	* Convert a Coordinate expressed in ECEF CRS to GEOGRAPHIC CRS
+	*/
+	UFUNCTION(BlueprintCallable, Category = "GeoReferencing|Transformations", meta = (DisplayName = "ECEF To Geographic"))
+	void K2_ECEFToGeographic(const FVector& ECEFCoordinates, FGeographicCoordinates& GeographicCoordinates) { ECEFToGeographic(ECEFCoordinates, GeographicCoordinates); }
+
+	// ENU & Transforms
+
+	/**
+	* Get the East North Up vectors at a specific location
+	*/
+	UFUNCTION(BlueprintCallable, Category = "GeoReferencing|ENU", meta = (DisplayName = "Get ENU Vectors At Projected Location"))
+	void K2_GetENUVectorsAtProjectedLocation(const FVector& ProjectedCoordinates, FVector& East, FVector& North, FVector& Up) { GetENUVectorsAtProjectedLocation(ProjectedCoordinates, East, North, Up); }
+
+	/**
+	* Get the East North Up vectors at a specific location
+	*/
+	UFUNCTION(BlueprintCallable, Category = "GeoReferencing|ENU", meta = (DisplayName = "Get ENU Vectors At ECEF Location"))
+	void K2_GetENUVectorsAtECEFLocation(const FVector& ECEFCoordinates, FVector& East, FVector& North, FVector& Up) { GetENUVectorsAtECEFLocation(ECEFCoordinates, East, North, Up); }
+
+	/**
+	* Get the East North Up vectors at a specific location - Not in engine frame, but in pure ECEF Frame !
+	*/
+	UFUNCTION(BlueprintCallable, Category = "GeoReferencing|ENU", meta = (DisplayName = "Get ECEF ENU Vectors At ECEF Location"))
+	void K2_GetECEFENUVectorsAtECEFLocation(const FVector& ECEFCoordinates, FVector& ECEFEast, FVector& ECEFNorth, FVector& ECEFUp) { GetECEFENUVectorsAtECEFLocation(ECEFCoordinates, ECEFEast, ECEFNorth, ECEFUp); }
+
+	/**
+	* Get the the transform to locate an object tangent to Ellipsoid at a specific location
+	*/
+	UFUNCTION(BlueprintCallable, Category = "GeoReferencing|TangentTransforms", meta = (DisplayName = "Get Tangent Transform At Projected Location"))
+	FTransform K2_GetTangentTransformAtProjectedLocation(const FVector& ProjectedCoordinates) { return GetTangentTransformAtProjectedLocation(ProjectedCoordinates); }
+
+	/**
+	* Get the the transform to locate an object tangent to Ellipsoid at a specific location
+	*/
+	UFUNCTION(BlueprintCallable, Category = "GeoReferencing|TangentTransforms", meta = (DisplayName = "Get Tangent Transform At ECEFLocation"))
+	FTransform K2_GetTangentTransformAtECEFLocation(const FVector& ECEFCoordinates) { return GetTangentTransformAtECEFLocation(ECEFCoordinates); }
+
+#pragma endregion
+
+#pragma region New Prototypes for C++
+
+	// Engine <--> Projected 
+
+	/**
+	* Convert a Vector expressed in ENGINE space to the PROJECTED CRS
+	*/
+	void EngineToProjected(const FVector& EngineCoordinates, FVector& ProjectedCoordinates);
+
+	/**
+	* Convert a Vector expressed in PROJECTED CRS to ENGINE space
+	*/
+	void ProjectedToEngine(const FVector& ProjectedCoordinates, FVector& EngineCoordinates);
+
+	// Engine <--> ECEF 
+
+	/**
+	* Convert a Vector expressed in ENGINE space to the ECEF CRS
+	*/
+	void EngineToECEF(const FVector& EngineCoordinates, FVector& ECEFCoordinates);
+
+	/**
+	* Convert a Vector expressed in ECEF CRS to ENGINE space
+	*/
+	void ECEFToEngine(const FVector& ECEFCoordinates, FVector& EngineCoordinates);
+
+	// Engine <--> Geographic 
+
+	/**
+	* Convert a Vector expressed in ENGINE space to the GEOGRAPHIC CRS
+	*/
 	UFUNCTION(BlueprintCallable, Category = "GeoReferencing|Transformations")
-	void GeographicToECEF(const FGeographicCoordinates& GeographicCoordinates, FCartesianCoordinates& ECEFCoordinates);
+	void EngineToGeographic(const FVector& EngineCoordinates, FGeographicCoordinates& GeographicCoordinates);
 
 	/**
-	 * Convert a Coordinate expressed in ECEF CRS to GEOGRAPHIC CRS
-	 */
+	* Convert a Vector expressed in GEOGRAPHIC CRS to ENGINE space
+	*/
 	UFUNCTION(BlueprintCallable, Category = "GeoReferencing|Transformations")
-	void ECEFToGeographic(const FCartesianCoordinates& ECEFCoordinates, FGeographicCoordinates& GeographicCoordinates);
+	void GeographicToEngine(const FGeographicCoordinates& GeographicCoordinates, FVector& EngineCoordinates);
 
 
-	// ENU & Transforms
+	// Projected <--> Geographic
 
 	/**
-	 * Get the East North Up vectors at a specific location
-	 */
+	* Convert a Coordinate expressed in PROJECTED CRS to GEOGRAPHIC CRS
+	*/
+	void ProjectedToGeographic(const FVector& ProjectedCoordinates, FGeographicCoordinates& GeographicCoordinates);
+
+	/**
+	* Convert a Coordinate expressed in GEOGRAPHIC CRS to PROJECTED CRS
+	*/
+	void GeographicToProjected(const FGeographicCoordinates& GeographicCoordinates, FVector& ProjectedCoordinates);
+
+	// Projected <--> ECEF
+
+	/**
+	* Convert a Coordinate expressed in PROJECTED CRS to ECEF CRS
+	*/
+	void ProjectedToECEF(const FVector& ProjectedCoordinates, FVector& ECEFCoordinates);
+
+	/**
+	* Convert a Coordinate expressed in ECEF CRS to PROJECTED CRS
+	*/
+	void ECEFToProjected(const FVector& ECEFCoordinates, FVector& ProjectedCoordinates);
+
+	// Geographic <--> ECEF
+
+	/**
+	* Convert a Coordinate expressed in GEOGRAPHIC CRS to ECEF CRS
+	*/
+	void GeographicToECEF(const FGeographicCoordinates& GeographicCoordinates, FVector& ECEFCoordinates);
+
+	/**
+	* Convert a Coordinate expressed in ECEF CRS to GEOGRAPHIC CRS
+	*/
+	void ECEFToGeographic(const FVector& ECEFCoordinates, FGeographicCoordinates& GeographicCoordinates);
+
+	// ENU and Transforms
+
+	/**
+	* Get the East North Up vectors at a specific location
+	*/
 	UFUNCTION(BlueprintCallable, Category = "GeoReferencing|ENU")
 	void GetENUVectorsAtEngineLocation(const FVector& EngineCoordinates, FVector& East, FVector& North, FVector& Up);
 
 	/**
-	 * Get the East North Up vectors at a specific location
-	 */
-	UFUNCTION(BlueprintCallable, Category = "GeoReferencing|ENU")
-	void GetENUVectorsAtProjectedLocation(const FCartesianCoordinates& ProjectedCoordinates, FVector& East, FVector& North, FVector& Up);
+	* Get the East North Up vectors at a specific location
+	*/
+	void GetENUVectorsAtProjectedLocation(const FVector& ProjectedCoordinates, FVector& East, FVector& North, FVector& Up);
 
 	/**
-	 * Get the East North Up vectors at a specific location
-	 */
+	* Get the East North Up vectors at a specific location
+	*/
 	UFUNCTION(BlueprintCallable, Category = "GeoReferencing|ENU")
 	void GetENUVectorsAtGeographicLocation(const FGeographicCoordinates& GeographicCoordinates, FVector& East, FVector& North, FVector& Up);
 
 	/**
-	 * Get the East North Up vectors at a specific location
-	 */
-	UFUNCTION(BlueprintCallable, Category = "GeoReferencing|ENU")
-	void GetENUVectorsAtECEFLocation(const FCartesianCoordinates& ECEFCoordinates, FVector& East, FVector& North, FVector& Up);
+	* Get the East North Up vectors at a specific location
+	*/
+	void GetENUVectorsAtECEFLocation(const FVector& ECEFCoordinates, FVector& East, FVector& North, FVector& Up);
 
 	/**
-	 * Get the East North Up vectors at a specific location - Not in engine frame, but in pure ECEF Frame !
-	 */
-	UFUNCTION(BlueprintCallable, Category = "GeoReferencing|ENU")
-	void GetECEFENUVectorsAtECEFLocation(const FCartesianCoordinates& ECEFCoordinates, FVector& ECEFEast, FVector& ECEFNorth, FVector& ECEFUp);
+	* Get the East North Up vectors at a specific location - Not in engine frame, but in pure ECEF Frame !
+	*/
+	void GetECEFENUVectorsAtECEFLocation(const FVector& ECEFCoordinates, FVector& ECEFEast, FVector& ECEFNorth, FVector& ECEFUp);
 
 	/**
-	 * Get the the transform to locate an object tangent to Ellipsoid at a specific location
-	 */
+	* Get the the transform to locate an object tangent to Ellipsoid at a specific location
+	*/
 	UFUNCTION(BlueprintCallable, Category = "GeoReferencing|TangentTransforms")
 	FTransform GetTangentTransformAtEngineLocation(const FVector& EngineCoordinates);
 
 	/**
-	 * Get the the transform to locate an object tangent to Ellipsoid at a specific location
-	 */
-	UFUNCTION(BlueprintCallable, Category = "GeoReferencing|TangentTransforms")
-	FTransform GetTangentTransformAtProjectedLocation(const FCartesianCoordinates& ProjectedCoordinates);
+	* Get the the transform to locate an object tangent to Ellipsoid at a specific location
+	*/
+	FTransform GetTangentTransformAtProjectedLocation(const FVector& ProjectedCoordinates);
 
 	/**
-	 * Get the the transform to locate an object tangent to Ellipsoid at a specific location
-	 */
+	* Get the the transform to locate an object tangent to Ellipsoid at a specific location
+	*/
 	UFUNCTION(BlueprintCallable, Category = "GeoReferencing|TangentTransforms")
 	FTransform GetTangentTransformAtGeographicLocation(const FGeographicCoordinates& GeographicCoordinates);
 
 	/**
-	 * Get the the transform to locate an object tangent to Ellipsoid at a specific location
-	 */
-	UFUNCTION(BlueprintCallable, Category = "GeoReferencing|TangentTransforms")
-	FTransform GetTangentTransformAtECEFLocation(const FCartesianCoordinates& ECEFCoordinates);
+	* Get the the transform to locate an object tangent to Ellipsoid at a specific location
+	*/
+	FTransform GetTangentTransformAtECEFLocation(const FVector& ECEFCoordinates);
 
 	/**
-	 * Set this transform to an Ellipsoid to have it positioned tangent to the origin.
-	 */
+	* Set this transform to an Ellipsoid to have it positioned tangent to the origin.
+	*/
 	UFUNCTION(BlueprintCallable, Category = "GeoReferencing|Misc")
 	FTransform GetPlanetCenterTransform();
 
 	// Public PROJ Utilities 
 
 	/**
-	 * Check if the string corresponds to a valid CRS descriptor
-	 */
+	* Check if the string corresponds to a valid CRS descriptor
+	*/
 	UFUNCTION(BlueprintCallable, Category = "GeoReferencing|Misc")
 	bool IsCRSStringValid(FString CRSString, FString& Error);
 
 
+	// Ellipsoid
+
+	/**
+	* Find the underlying Geographic CRS Ellipsoid and return its maximum radius
+	*/
+	UFUNCTION(BlueprintCallable, Category = "GeoReferencing|Ellipsoid")
+	double GetGeographicEllipsoidMaxRadius();
+	/**
+	* Find the underlying Geographic CRS Ellipsoid and return its minimum radius
+	*/
+	UFUNCTION(BlueprintCallable, Category = "GeoReferencing|Ellipsoid")
+	double GetGeographicEllipsoidMinRadius();
+
+	/**
+	* Find the underlying Projected CRS Ellipsoid and return its maximum radius
+	*/
+	UFUNCTION(BlueprintCallable, Category = "GeoReferencing|Ellipsoid")
+	double GetProjectedEllipsoidMaxRadius();
+	/**
+	* Find the underlying Projected CRS Ellipsoid and return its minimum radius
+	*/
+	UFUNCTION(BlueprintCallable, Category = "GeoReferencing|Ellipsoid")
+	double GetProjectedEllipsoidMinRadius();
+
+
+
+	/**
+	* In case you want to change the Origin or CRS definition properties during application execution, you need to call this function to update the internal transformation cache.
+	* Note this is not a recommended practice, because it will not move the level actors accordingly.
+	* Can be useful though if you rebase your actors yourself, or if you want to change one CRS used for displaying coordinates.
+	*/
+	UFUNCTION(BlueprintCallable, Category = "GeoReferencing|Misc", meta = (DisplayName = "Apply Runtime Changes"))
+	void ApplySettings();
+
+#pragma endregion
+
 	//////////////////////////////////////////////////////////////////////////
 	// General
 
 	/**
-	 * This mode has to be set consistently with the way you authored your ground geometry.
-	 *  - For small environments, a projection is often applied and the world is considered as Flat
-	 *  - For planet scale environments, projections is not suitable and the geometry is Rounded.
-	 **/
+	* This mode has to be set consistently with the way you authored your ground geometry.
+	*  - For small environments, a projection is often applied and the world is considered as Flat
+	*  - For planet scale environments, projections is not suitable and the geometry is Rounded.
+	**/
 	UPROPERTY(EditAnywhere, BlueprintReadWrite, Category = "GeoReferencing")
 	EPlanetShape PlanetShape = EPlanetShape::FlatPlanet;
 
 	/**
-	 * String that describes the PROJECTED CRS of choice. 
-	 *    CRS can be identified by their code (EPSG:4326), a well-known text(WKT) string, or PROJ strings...
-	 **/
+	* String that describes the PROJECTED CRS of choice.
+	*    CRS can be identified by their code (EPSG:4326), a well-known text(WKT) string, or PROJ strings...
+	**/
 	UPROPERTY(EditAnywhere, BlueprintReadWrite, Config, Category = "GeoReferencing")
 	FString ProjectedCRS = FString(TEXT("EPSG:32631")); // UTM Zone 31 North - (0,0) https://epsg.io/32631
 
 	/**
-	 * String that describes the GEOGRAPHIC CRS of choice.
-	 *    CRS can be identified by their code (EPSG:4326), a well-known text(WKT) string, or PROJ strings...
-	 **/
+	* String that describes the GEOGRAPHIC CRS of choice.
+	*    CRS can be identified by their code (EPSG:4326), a well-known text(WKT) string, or PROJ strings...
+	**/
 	UPROPERTY(EditAnywhere, BlueprintReadWrite, Config, Category = "GeoReferencing")
 	FString GeographicCRS = FString(TEXT("EPSG:4326")); // WGS84 https://epsg.io/4326
 
@@ -218,62 +436,61 @@ public:
 	// Origin Location
 
 	/**
-	 * if true, the UE origin is located at the Planet Center, otherwise, 
-	 * the UE origin is assuming to be defined at one specific point of 
-	 * the planet surface, defined by the properties below. 
-	 **/
+	* if true, the UE origin is located at the Planet Center, otherwise,
+	* the UE origin is assuming to be defined at one specific point of
+	* the planet surface, defined by the properties below.
+	**/
 	UPROPERTY(EditAnywhere, BlueprintReadWrite, Category = "GeoReferencing|Origin Location", meta = (EditConditionHides, EditCondition = "PlanetShape==EPlanetShape::RoundPlanet"))
 	bool bOriginAtPlanetCenter = false;
 
 	/**
-	 * if true, the UE origin georeference is expressed in the PROJECTED CRS. 
-	 *     (NOT in ECEF ! - Projected worlds are the most frequent use case...)
-	 * if false, the origin is located using geographic coordinates. 
-	 * 
-	 * WARNING : the location has to be expressed as Integer values because of accuracy. 
-	 * Be very careful about that when authoring your data in external tools ! 
-	 **/
+	* if true, the UE origin georeference is expressed in the PROJECTED CRS.
+	*     (NOT in ECEF ! - Projected worlds are the most frequent use case...)
+	* if false, the origin is located using geographic coordinates.
+	*
+	* WARNING : the location has to be expressed as Integer values because of accuracy.
+	* Be very careful about that when authoring your data in external tools !
+	**/
 	UPROPERTY(EditAnywhere, BlueprintReadWrite, Category = "GeoReferencing|Origin Location", meta = (EditConditionHides, EditCondition = "bOriginAtPlanetCenter==false"))
 	bool bOriginLocationInProjectedCRS = true;
 
 	/**
-	 * Latitude of UE Origin on planet
-	 **/
+	* Latitude of UE Origin on planet
+	**/
 	UPROPERTY(EditAnywhere, BlueprintReadWrite, Category = "GeoReferencing|Origin Location", meta = (EditConditionHides, EditCondition = "!bOriginLocationInProjectedCRS && !bOriginAtPlanetCenter", ClampMin = "-90", ClampMax = "90"))
 	double OriginLatitude = 0.0;
 
-
 	/**
-	 * Longitude of UE Origin on planet
-	 **/
+	* Longitude of UE Origin on planet
+	**/
 	UPROPERTY(EditAnywhere, BlueprintReadWrite, Category = "GeoReferencing|Origin Location", meta = (EditConditionHides, EditCondition = "!bOriginLocationInProjectedCRS && !bOriginAtPlanetCenter", ClampMin = "-180", ClampMax = "180"))
 	double OriginLongitude = 0.0;
 
 	/**
-	 * Altitude of UE Origin on planet
-	 **/
+	* Altitude of UE Origin on planet
+	**/
 	UPROPERTY(EditAnywhere, BlueprintReadWrite, Category = "GeoReferencing|Origin Location", meta = (EditConditionHides, EditCondition = "!bOriginLocationInProjectedCRS && !bOriginAtPlanetCenter"))
-	double OriginAltitude = 0.0; // For constituency, but should be 0
+	double OriginAltitude = 0.0;
 
 	/**
-	 * Easting position of UE Origin on planet, express in the Projected CRS Frame
-	 **/
+	* Easting position of UE Origin on planet, express in the Projected CRS Frame
+	**/
 	UPROPERTY(EditAnywhere, BlueprintReadWrite, Category = "GeoReferencing|Origin Location", meta = (EditConditionHides, EditCondition = "bOriginLocationInProjectedCRS && !bOriginAtPlanetCenter"))
 	double OriginProjectedCoordinatesEasting = 500000.0;
 
 	/**
-	 * Northing position of UE Origin on planet, express in the Projected CRS Frame
-	 **/
+	* Northing position of UE Origin on planet, express in the Projected CRS Frame
+	**/
 	UPROPERTY(EditAnywhere, BlueprintReadWrite, Category = "GeoReferencing|Origin Location", meta = (EditConditionHides, EditCondition = "bOriginLocationInProjectedCRS && !bOriginAtPlanetCenter"))
 	double OriginProjectedCoordinatesNorthing = 5000000.0;
 
 	/**
-	 * Up position of UE Origin on planet, express in the Projected CRS Frame
-	 **/
+	* Up position of UE Origin on planet, express in the Projected CRS Frame
+	**/
 	UPROPERTY(EditAnywhere, BlueprintReadWrite, Category = "GeoReferencing|Origin Location", meta = (EditConditionHides, EditCondition = "bOriginLocationInProjectedCRS && !bOriginAtPlanetCenter"))
 	double OriginProjectedCoordinatesUp = 0.0;
 
-	
+
 #if WITH_EDITOR
 	virtual void PostEditChangeProperty(FPropertyChangedEvent& PropertyChangedEvent) override;
 #endif
@@ -281,9 +498,6 @@ public:
 private:
 	void Initialize();
 
-public:
-	void ApplySettings();
-
 private:
 	class FGeoReferencingSystemInternals;
 	TPimplPtr<FGeoReferencingSystemInternals> Impl;
diff --git a/Engine/Plugins/Runtime/GeoReferencing/Source/GeoReferencing/Public/GeographicCoordinates.h b/Engine/Plugins/Runtime/GeoReferencing/Source/GeoReferencing/Public/GeographicCoordinates.h
--- a/Engine/Plugins/Runtime/GeoReferencing/Source/GeoReferencing/Public/GeographicCoordinates.h
+++ b/Engine/Plugins/Runtime/GeoReferencing/Source/GeoReferencing/Public/GeographicCoordinates.h
@@ -3,6 +3,7 @@
 #pragma once
 
 #include "CoreMinimal.h"
+#include "GeoReferencingModule.h"
 #include "Kismet/BlueprintFunctionLibrary.h"
 
 #include "GeographicCoordinates.generated.h"
@@ -16,6 +17,7 @@ public:
 
 	FGeographicCoordinates();
 	FGeographicCoordinates(double InLongitude, double InLatitude, double InAltitude);
+	FGeographicCoordinates(const FVector& LatLongAltVector); // FVector where X = Latitude, Y = Longitude, Z = Altitude
 
 	UPROPERTY(EditAnywhere, BlueprintReadWrite, Config, Category = "GeoReferencing")
 	double Longitude;
@@ -68,7 +70,8 @@ class GEOREFERENCING_API UGeographicCoordinatesFunctionLibrary : public UBluepri
 	 * Get the Coordinates as a float approximation.
 	 * USE WISELY as we can't guarantee there will no be rounding due to IEEE754 float encoding !
 	 **/
-	UFUNCTION(BlueprintPure, Category = "GeoReferencing")
+	UE_DEPRECATED(5.0, "BP now support doubles, Function useless and can lead to precision issues")
+	UFUNCTION(BlueprintPure, Category = "GeoReferencing", meta = (DeprecatedFunction, DeprecationMessage = "BP now support doubles, Function useless and can lead to precision issues"))
 	static void ToFloatApproximation(UPARAM(ref) FGeographicCoordinates& GeographicCoordinates, float& OutLatitude, float& OutLongitude, float& OutAltitude)
 	{
 		OutLatitude  = static_cast<float>(GeographicCoordinates.Latitude);
@@ -80,9 +83,31 @@ class GEOREFERENCING_API UGeographicCoordinatesFunctionLibrary : public UBluepri
 	 * Set the Coordinates from float approximation.
 	 * USE WISELY as we can't guarantee there will no be rounding due to IEEE754 float encoding !
 	 **/
-	UFUNCTION(BlueprintPure, Category = "GeoReferencing|Coordinates")
+	UE_DEPRECATED(5.0, "BP now support doubles, Function useless and can lead to precision issues")
+	UFUNCTION(BlueprintPure, Category = "GeoReferencing|Coordinates", meta = (DeprecatedFunction, DeprecationMessage = "BP now support doubles, Function useless and can lead to precision issues"))
 	static FGeographicCoordinates MakeGeographicCoordinatesApproximation(const float& InLatitude, const float& InLongitude, const float& InAltitude)
 	{
 		return FGeographicCoordinates(static_cast<double>(InLongitude), static_cast<double>(InLatitude), static_cast<double>(InAltitude));
 	}
+
+
+	/**
+	 * Make Geographic Coordinates from a FVector where X=Latitude, Y=Longitude, Z=Altitude
+	 **/
+	UFUNCTION(BlueprintPure, Category = "GeoReferencing|Coordinates")
+	static FGeographicCoordinates MakeGeographicCoordinates(const FVector& LatLongAltVector)
+	{
+		return FGeographicCoordinates(LatLongAltVector);
+	}
+
+	/**
+	 * Express the Geographic coordinates as a FVector where  where X=Latitude, Y=Longitude, Z=Altitude
+	 **/
+	UFUNCTION(BlueprintPure, Category = "GeoReferencing", meta = (DeprecatedFunction, DeprecationMessage = "BP now support doubles, Function useless and can lead to precision issues"))
+	static void ToLatLongAltVector(UPARAM(ref) FGeographicCoordinates& GeographicCoordinates, FVector& OutLatLongAltVector)
+	{
+		OutLatLongAltVector = FVector(GeographicCoordinates.Latitude, GeographicCoordinates.Longitude, GeographicCoordinates.Altitude);
+	}
+
+
 };
diff --git a/Engine/Plugins/Runtime/GeoReferencing/Source/GeoReferencing/Public/RoundPlanetPawn.h b/Engine/Plugins/Runtime/GeoReferencing/Source/GeoReferencing/Public/RoundPlanetPawn.h
new file mode 100644
--- /dev/null
+++ b/Engine/Plugins/Runtime/GeoReferencing/Source/GeoReferencing/Public/RoundPlanetPawn.h
@@ -0,0 +1,269 @@
+// Copyright Epic Games, Inc. All Rights Reserved.
+
+#pragma once
+
+#include "CoreMinimal.h"
+
+#include "GeographicCoordinates.h"
+#include "GameFramework/DefaultPawn.h"
+#include "RoundPlanetPawn.generated.h"
+
+
+class AGeoReferencingSystem;
+class UCurveFloat;
+
+
+/**
+ * This pawn can be used to easily move around the globe while maintaining a
+ * sensible orientation. As the pawn moves across the horizon, it automatically
+ * changes its own up direction such that the world always looks right-side up.
+ */
+UCLASS()
+class GEOREFERENCING_API ARoundPlanetPawn : public ADefaultPawn
+{
+	GENERATED_BODY()
+	ARoundPlanetPawn();
+
+public:
+
+	/**
+	* This curve dictates what percentage of the max altitude the pawn should take at a given time on the curve.
+	* Depending on the total distance to Fly, a constant maximum altitude is computed from FlyToLocationMaximumAltitudeCurve
+	* Then during the fly, the actual pawn altitude is computed from this profile and the reference maximum altitude. 
+	* This allows to have nice movements whenever you FlyTo near or far from your actual location. 
+	* This curve must be kept in the 0 to 1 range on both axes.
+	*/
+	UPROPERTY(EditAnywhere, Category = "FlyToLocation")
+	UCurveFloat* AltitudeProfileCurve;
+
+	/**
+	* This curve is used to pick up a reference the maximum altitude when flying to a location. 
+	* This maximum altitude will be moderated by the FlyToLocationAltitudeProfileCurve
+	* X Axis is the distance of the flight (meters)
+	* Y Axis is the maximum altitude to be used for this flight (meters)
+	*/
+	UPROPERTY(EditAnywhere, Category = "FlyToLocation")
+	UCurveFloat* MaximumAltitudeCurve;
+
+	/**
+	* In case MaximumAltitudeCurve is not defined, the AltitudeProfileCurve will use this Maximum altitude value for the flight, whatever the travel distance. 
+	* In meters
+	*/
+	UPROPERTY(EditAnywhere, Category = "FlyToLocation")
+	float MaximumAltitudeValue = 300;
+
+	/** 
+	* This curve is used to ease in an out the Fly to Location speed.
+	* It must be kept in the 0 to 1 range on both axes.
+	* X axis is the normalized time of travel 
+	* Y axis is the progress on the flight motion. 
+	* A linear curve will give a constant progress speed whereas a easein//easeout one will slow down the speed at beginning/end of motion. 
+	*/
+	UPROPERTY(EditAnywhere, Category = "FlyToLocation")
+	UCurveFloat* ProgressCurve;
+
+	/**
+	* Fly to Location duration (in seconds)
+	*/
+	UPROPERTY(EditAnywhere, Category = "FlyToLocation", meta = (ClampMin = 0.0))
+	float FlyDuration = 5.0;
+
+	/**
+	* The granularity in degrees with which keypoints should be generated for the Fly to Location interpolation. 
+	* Lower values means smoother motion. 
+	* ROM : 1 degree latitude ~110km on earth. => 0.1 ~10km 
+	*/
+	UPROPERTY(EditAnywhere, Category = "FlyToLocation", meta = (ClampMin = 0.0))
+	float GranularityDegrees = 0.1;
+
+	/**
+	* The minimum linear steps for the FlyTolocation motion. 
+	* GranularityDegrees is not sufficient in case of short travels. Eg if 0.1, any jump shorter than 10km will have only one step. 
+	* Make sure we have at least this number of steps in the trajectory. 
+	*/
+	UPROPERTY(EditAnywhere, Category = "FlyToLocation", meta = (ClampMin = 0.0))
+	int32 MinimumStepCount = 10;
+
+	/**
+	* Make sure we don't get crazy in case of large flights with small granularity
+	*/
+	UPROPERTY(EditAnywhere, Category = "FlyToLocation", meta = (ClampMin = 0.0))
+	int32 MaximumStepCount = 300;
+
+
+	/**
+	* if True, the motion Forward/Right motion of the pawn are relative to Planet tangent, 
+	* meaning the altitude will approximately be kept, whatever the pawn camera orientation
+	*/
+	UPROPERTY(EditAnywhere, BlueprintReadWrite, Category = "Movement")
+	bool OrbitalMotion = false;
+
+	/**
+	* The Reference maximum speed for the pawn, before being altered by any Scalar modifier or by altitude curve
+	* ActualMaxSpeed = BaseSpeedKmh * SpeedScalar * AltitudeSpeedModifierCurve(Altitude)
+	*/
+	UPROPERTY(EditAnywhere, BlueprintReadWrite, Category = "Movement")
+	float BaseSpeedKmh = 100.0f;
+
+	/**
+	* Scalar modifier for the base speed
+	* ActualMaxSpeed = BaseSpeedKmh * SpeedScalar * AltitudeSpeedModifierCurve(Altitude)
+	*/
+	UPROPERTY(EditAnywhere, BlueprintReadWrite, Category = "Movement")
+	float SpeedScalar = 1.0f;
+
+	/**
+	* Multiplier/Divider for increasing/decreasing the speed scalar
+	*/
+	UPROPERTY(EditAnywhere, BlueprintReadWrite, Category = "Movement")
+	float SpeedScalarIncrement = 1.5f;
+
+	/**
+	* When being very high and moving around the planet, we have to dynamically increase the speed based on Height above terrain to accelerate movement
+	* This curve adds a multiplying factor to the maximum speed depending on HAT. 
+	* ActualMaxSpeed = BaseSpeedKmh * SpeedScalar * AltitudeSpeedModifierCurve(Altitude)
+	*/
+	UPROPERTY(EditAnywhere, Category = "Movement")
+	UCurveFloat* SpeedByHATModifierCurve;
+
+	/**
+	* Height Above Terrain. The distance between the ground and the pawn
+	*/
+	UPROPERTY(VisibleAnywhere, BlueprintReadOnly, Category = "Movement")
+	float HAT = 0;
+
+	/**
+	* The distance between the geographic ellipsoid surface and the pawn
+	*/
+	UPROPERTY(VisibleAnywhere, BlueprintReadOnly, Category = "Movement")
+	float Altitude = 0;
+
+public:
+	/**
+	* Begin a smooth camera flight to the given ECEF destination such that the camera ends at the specified yaw and pitch.
+	* The flight can be enforced or canceled if the user moves the pawn
+	*/
+	UFUNCTION(BlueprintCallable, Category = "FlyToLocation")
+	void FlyToLocationECEF(const FVector& ECEFDestination, double YawAtDestination, double PitchAtDestination, bool CanInterruptByMoving);
+
+	/**
+	* Begin a smooth camera flight to the given Latitude/Longitude destination such that the camera ends at the specified yaw and pitch.
+	* The flight can be enforced or canceled if the user moves the pawn
+	*/
+	UFUNCTION(BlueprintCallable, Category = "FlyToLocation")
+	void FlyToLocationGeographic(const FGeographicCoordinates& GeographicDestination, double YawAtDestination, double PitchAtDestination, bool CanInterruptByMoving);
+
+	/**
+	* Begin a smooth camera flight to the given Latitude/Longitude destination such that the camera ends at the specified yaw and pitch.
+	* The flight can be enforced or canceled if the user moves the pawn
+	*/
+	UFUNCTION(BlueprintCallable, Category = "FlyToLocation")
+	void FlyToLocationLatitudeLongitudeAltitude(const double& InLatitude, const double& InLongitude, const double& InAltitude, double YawAtDestination, double PitchAtDestination, bool CanInterruptByMoving);
+
+	/**
+	* Stop the current Fly To Location motion
+	*/
+	UFUNCTION(BlueprintCallable, Category = "FlyToLocation")
+	void InterruptFlyToLocation();
+
+	/**
+	* Reset the Speed Scalar to its default value - Middle-mouse button click equivalent
+	*/
+	UFUNCTION(BlueprintCallable, Category = "Movement")
+	void ResetSpeedScalar() { SpeedScalar = 1.0;};
+
+	/**
+	* Increase the Speed Scalar - MouseWheel Up equivalent
+	*/
+	UFUNCTION(BlueprintCallable, Category = "Movement")
+	void IncreaseSpeedScalar() { SpeedScalar *= SpeedScalarIncrement; }
+
+	/**
+	* Decrease the Speed Scalar - MouseWheel Down equivalent
+	*/
+	UFUNCTION(BlueprintCallable, Category = "Movement")
+	void DecreaseSpeedScalar() { SpeedScalar /= SpeedScalarIncrement; }
+
+	// Public Overridables
+
+	void OnConstruction(const FTransform& Transform) override;
+
+	void BeginPlay() override;
+
+	void Tick(float DeltaSeconds) override;
+
+	/**
+	* Input callback to move forward in local space (or backward if Val is negative).
+	* @param Val Amount of movement in the forward direction (or backward if negative).
+	* @see APawn::AddMovementInput()
+	*/
+	virtual void MoveForward(float Val) override;
+
+	/**
+	* Input callback to strafe right in local space (or left if Val is negative).
+	* @param Val Amount of movement in the right direction (or left if negative).
+	* @see APawn::AddMovementInput()
+	*/
+	virtual void MoveRight(float Val) override;
+
+	/**
+	* Input callback to move up in world space (or down if Val is negative).
+	* @param Val Amount of movement in the world up direction (or down if negative).
+	* @see APawn::AddMovementInput()
+	*/
+	virtual void MoveUp_World(float Val) override;
+
+
+	/**
+	* Don't return the controller rotation but a transformed rotation to consider planet roundness. 
+	*/
+	virtual FRotator GetViewRotation() const override;
+
+	/**
+	* Don't return the controller rotation but a transformed rotation to consider planet roundness. 
+	*/
+	virtual FRotator GetBaseAimRotation() const override;
+
+	/**
+	* Adds to the default pawn the Mouse Wheel inputs to tune the speed scalar
+	*/
+	void SetupPlayerInputComponent(UInputComponent* InInputComponent) override;
+
+protected:
+	UPROPERTY()
+	AGeoReferencingSystem* GeoReferencingSystem;
+
+private:
+	/**
+	 * In order to limit computations, we keep a cache for some values : TangentTransform, ENU, Altitude, HAT
+	 * They are updated each time the pawn travels more than CacheUpdateThresholdSq (In UE Units). 
+	 * if not HAT is found, we retry each tick. 
+	*/
+	void UpdateMotionObjectsCache(bool ForceUpdate = false);
+	/**
+	 * Updates the Movement Component Speed and Inertia based on the speed scalar changes. 
+	*/
+	void HandleMotionSpeed();
+	/**
+	* When flying to a specific location, this function computes at each tick the Planet-relative trajectory. 
+	*/
+	void ProcessFlyToLocation(float DeltaSeconds);
+	
+
+	// Variables for FlyToLocation
+	bool bIsFlyingToLocation = false;
+	bool bCanInterruptFlight = false;
+	float CurrentFlyTime = 0.0;
+	FRotator FlyToLocationOriginRotation;
+	FRotator FlyToLocationDestinationRotation;
+	TArray<FVector> FlyToLocationKeypoints;
+	
+	// Variables for Motion Object caches
+	FVector LastCacheLocation;
+	float CacheUpdateThresholdSq = 200.0 * 200.0; // 2m
+	FVector East;
+	FVector North;
+	FVector Up;
+	FTransform TangentTransform;
+	bool bIsHATInvalid = true;
+};
