Fixed compilation error with Unity builds

### Notes
Added missing headers
related to HORDE Issue 137259

### Hashtags

diff --git a/Engine/Plugins/Runtime/GeoReferencing/Source/GeoReferencing/Private/GeoReferencingModule.cpp b/Engine/Plugins/Runtime/GeoReferencing/Source/GeoReferencing/Private/GeoReferencingModule.cpp
--- a/Engine/Plugins/Runtime/GeoReferencing/Source/GeoReferencing/Private/GeoReferencingModule.cpp
+++ b/Engine/Plugins/Runtime/GeoReferencing/Source/GeoReferencing/Private/GeoReferencingModule.cpp
@@ -1,7 +1,12 @@
 // Copyright Epic Games, Inc. All Rights Reserved.
 
 #include "GeoReferencingModule.h"
+#include "Modules/ModuleManager.h"
 
-GEOREFERENCING_API DEFINE_LOG_CATEGORY(LogGeoReferencing);
+DEFINE_LOG_CATEGORY(LogGeoReferencing);
 
 IMPLEMENT_MODULE(FGeoReferencingModule, GeoReferencing)
+
+
+
+
diff --git a/Engine/Plugins/Runtime/GeoReferencing/Source/GeoReferencing/Private/RoundPlanetPawn.cpp b/Engine/Plugins/Runtime/GeoReferencing/Source/GeoReferencing/Private/RoundPlanetPawn.cpp
--- a/Engine/Plugins/Runtime/GeoReferencing/Source/GeoReferencing/Private/RoundPlanetPawn.cpp
+++ b/Engine/Plugins/Runtime/GeoReferencing/Source/GeoReferencing/Private/RoundPlanetPawn.cpp
@@ -2,9 +2,15 @@
 
 
 #include "RoundPlanetPawn.h"
-#include "GameFramework/FloatingPawnMovement.h"
 #include "GeoReferencingSystem.h"
+#include "GameFramework/FloatingPawnMovement.h"
 #include "GameFramework/PlayerInput.h"
+#include "GameFramework/Actor.h"
+#include "GameFramework/Pawn.h"
+#include "GameFramework/WorldSettings.h"
+#include "GameFramework/Controller.h"
+#include "Components/InputComponent.h"
+#include "Curves/CurveFloat.h"
 
 
 ARoundPlanetPawn::ARoundPlanetPawn() : ADefaultPawn()
diff --git a/Engine/Plugins/Runtime/GeoReferencing/Source/GeoReferencing/Public/GeoReferencingModule.h b/Engine/Plugins/Runtime/GeoReferencing/Source/GeoReferencing/Public/GeoReferencingModule.h
--- a/Engine/Plugins/Runtime/GeoReferencing/Source/GeoReferencing/Public/GeoReferencingModule.h
+++ b/Engine/Plugins/Runtime/GeoReferencing/Source/GeoReferencing/Public/GeoReferencingModule.h
@@ -2,11 +2,14 @@
 
 #pragma once
 
-#include "CoreMinimal.h"
 #include "Modules/ModuleInterface.h"
+#include "CoreMinimal.h"
 
 GEOREFERENCING_API DECLARE_LOG_CATEGORY_EXTERN(LogGeoReferencing, Log, All);
 
+/**
+ *  FGeoReferencingModule
+ */
 class GEOREFERENCING_API FGeoReferencingModule : public IModuleInterface
 {
-};
+};
\ No newline at end of file
