Remove HoloLens from Editor module target, as there is no HoloLens Editor.

#rnx

diff --git a/Engine/Plugins/Runtime/GeoReferencing/GeoReferencing.uplugin b/Engine/Plugins/Runtime/GeoReferencing/GeoReferencing.uplugin
--- a/Engine/Plugins/Runtime/GeoReferencing/GeoReferencing.uplugin
+++ b/Engine/Plugins/Runtime/GeoReferencing/GeoReferencing.uplugin
@@ -33,7 +33,6 @@
 			"LoadingPhase": "Default",
 			"PlatformAllowList": [
 				"Android",
-				"HoloLens",
 				"IOS",
 				"Linux",
 				"Mac",
