Fixed the georeferencing plugin to make it work inside a sandboxed environment (Dedicated Server / Client)

diff --git a/Engine/Plugins/Runtime/GeoReferencing/Source/GeoReferencing/Private/GeoReferencingSystem.cpp b/Engine/Plugins/Runtime/GeoReferencing/Source/GeoReferencing/Private/GeoReferencingSystem.cpp
--- a/Engine/Plugins/Runtime/GeoReferencing/Source/GeoReferencing/Private/GeoReferencingSystem.cpp
+++ b/Engine/Plugins/Runtime/GeoReferencing/Source/GeoReferencing/Private/GeoReferencingSystem.cpp
@@ -10,6 +10,8 @@
 #include "Widgets/Notifications/SNotificationList.h"
 #include "Framework/Notifications/NotificationManager.h"
 #include "Engine/Engine.h"
+#include "HAL/FileManager.h"
+#include "HAL/PlatformFileManager.h"
 #include "GeographicCoordinates.h"
 #include "MathUtil.h"
  
@@ -899,6 +901,32 @@ void AGeoReferencingSystem::FGeoReferencingSystemInternals::InitPROJLibrary()
 	proj_context_set_search_paths(ProjContext, sizeof(ProjSearchPaths)/sizeof(ProjSearchPaths[0]), ProjSearchPaths);
 
 	proj_context_set_autoclose_database(ProjContext, true);
+	
+
+	// Special case for Dedicated Server or client
+	// When running in a sandboxed environment, the PROJ data files are not staged in that environment by the Cooking process. 
+	// And our SQLite accessor will always want to find these files in the sandboxed location, never in the original one. 
+	// And we don't want to have to build the PAK files when iterating on the code
+	// Let's manually populate the sandbox environment. 
+	IPlatformFile& PlatformFile = FPlatformFileManager::Get().GetPlatformFile();
+	if (PlatformFile.IsSandboxEnabled())
+	{
+		// Ask the system for the expected location of the proj.db file, using ConvertToAbsolutePathForExternalAppForWrite (not ***Read, that falls back on the original location)
+		FString ProjDBFilePath = FPaths::Combine(*ProjDataPath, TEXT("proj.db"));
+		FString ProjDBFileSandboxPath = PlatformFile.ConvertToAbsolutePathForExternalAppForWrite(*ProjDBFilePath);
+		UE_LOG(LogGeoReferencing, Warning, TEXT("Running in Sandboxed mode with Proj.db expected in %s "), *ProjDBFileSandboxPath);
+
+		// Check if the file is missing there
+		if (!IFileManager::Get().FileExists(*ProjDBFileSandboxPath))
+		{
+			UE_LOG(LogGeoReferencing, Warning, TEXT("Missing PROJ data files in then Sandbox environment --> Copy them!"));
+
+			// Create the target dir, and deploy the content. 
+			FString TargetFolder = IFileManager::Get().ConvertToAbsolutePathForExternalAppForWrite(*ProjDataPath); // Use **Write for the actual Sandboxed path. 
+			PlatformFile.GetLowerLevel()->CreateDirectoryTree(*TargetFolder); // Use lower level Platformfile (otherwise the sandbox path is added twice)
+			PlatformFile.GetLowerLevel()->CopyDirectoryTree(*TargetFolder, *ProjDataPath, false);
+		}
+	}
 	// Non-editor builds use UFS extensions to read PROJ data from UFS/Pak
 	if (!GIsEditor)
 	{
