Updated GeoReferencingDemo to UE 5.3
- Fixed BP using deprecated Float to Double nodes
- Removed the deprecated Rebasing Actor
- Removed the dependency to World->OriginLocation

diff --git a/Engine/Plugins/Runtime/GeoReferencing/Source/GeoReferencing/Private/GeoReferencingSystem.cpp b/Engine/Plugins/Runtime/GeoReferencing/Source/GeoReferencing/Private/GeoReferencingSystem.cpp
--- a/Engine/Plugins/Runtime/GeoReferencing/Source/GeoReferencing/Private/GeoReferencingSystem.cpp
+++ b/Engine/Plugins/Runtime/GeoReferencing/Source/GeoReferencing/Private/GeoReferencingSystem.cpp
@@ -248,15 +248,8 @@ void AGeoReferencingSystem::EngineToProjected(const FVector& EngineCoordinates,
 	default:
 	{
 		// in FlatPlanet, the transform is simply a translation
-		// Before any conversion, consider the internal UE rebasing
-		FIntVector UERebasingOffset = GetWorld()->OriginLocation;
-		FVector UERebasedCoordinates(
-			EngineCoordinates.X + UERebasingOffset.X,
-			EngineCoordinates.Y + UERebasingOffset.Y,
-			EngineCoordinates.Z + UERebasingOffset.Z);
-
 		// Convert UE units to meters, invert the Y coordinate because of left-handed UE Frame
-		FVector UEWorldCoordinates = UERebasedCoordinates * FVector(0.01, -0.01, 0.01);
+		FVector UEWorldCoordinates = EngineCoordinates * FVector(0.01, -0.01, 0.01);
 
 		// Add the defined origin offset
 		ProjectedCoordinates = UEWorldCoordinates + Impl->WorldOriginLocationProjected;
@@ -286,11 +279,7 @@ void AGeoReferencingSystem::ProjectedToEngine(const FVector& ProjectedCoordinate
 		FVector UEWorldCoordinates = (ProjectedCoordinates - Impl->WorldOriginLocationProjected);
 
 		// Convert UE units to meters, invert the Y coordinate because of left-handed UE Frame
-		FVector UERebasedCoordinates = UEWorldCoordinates * FVector(100.0, -100.0, 100.0);
-
-		// Consider the UE internal rebasing
-		FIntVector UERebasingOffset = GetWorld()->OriginLocation;
-		EngineCoordinates = UERebasedCoordinates - FVector(UERebasingOffset.X, UERebasingOffset.Y, UERebasingOffset.Z);
+		EngineCoordinates = UEWorldCoordinates * FVector(100.0, -100.0, 100.0);
 	}
 	break;
 	}
@@ -302,19 +291,12 @@ void AGeoReferencingSystem::EngineToECEF(const FVector& EngineCoordinates, FVect
 	{
 	case EPlanetShape::RoundPlanet:
 	{
-		// Before any conversion, consider the internal UE rebasing
-		FIntVector UERebasingOffset = GetWorld()->OriginLocation;
-		FVector UERebasedCoordinates(
-			EngineCoordinates.X + UERebasingOffset.X,
-			EngineCoordinates.Y + UERebasingOffset.Y,
-			EngineCoordinates.Z + UERebasingOffset.Z);
-
 		// Convert UE units to meters, invert the Y coordinate because of left-handed UE Frame
-		FVector UEWorldCoordinates = UERebasedCoordinates * FVector(0.01, -0.01, 0.01);
+		FVector UEWorldCoordinates = EngineCoordinates * FVector(0.01, -0.01, 0.01);
 
 		if (bOriginAtPlanetCenter)
 		{
-			// Easy case, UE is ECEF... And we did the rebasing, so return the Global coordinates
+			// Easy case, UE is ECEF... return the Global coordinates
 			ECEFCoordinates = UEWorldCoordinates; // TOCHECK Types
 		}
 		else
@@ -354,11 +336,7 @@ void AGeoReferencingSystem::ECEFToEngine(const FVector& ECEFCoordinates, FVector
 		}
 
 		// Convert UE units to meters, invert the Y coordinate because of left-handed UE Frame
-		FVector UERebasedCoordinates = UEWorldCoordinates * FVector(100.0, -100.0, 100.0);
-
-		// Consider the UE internal rebasing
-		FIntVector UERebasingOffset = GetWorld()->OriginLocation;
-		EngineCoordinates = UERebasedCoordinates - FVector(UERebasingOffset.X, UERebasingOffset.Y, UERebasingOffset.Z);
+		EngineCoordinates = UEWorldCoordinates * FVector(100.0, -100.0, 100.0);
 	}
 	break;
 
@@ -639,12 +617,7 @@ FTransform AGeoReferencingSystem::GetPlanetCenterTransform()
 
 			// Get Origin, and convert UE units to meters, invert the Y coordinate because of left-handed UE Frame
 			FVector UEOrigin = TransformMatrix.GetOrigin() * FVector(100.0, -100.0, 100.0);
-			
-			// Consider the UE internal rebasing to compute Origin
-			FIntVector UERebasingOffset = GetWorld()->OriginLocation;
-			FVector UERebasedCoordinates = UEOrigin - FVector(UERebasingOffset.X, UERebasingOffset.Y, UERebasingOffset.Z);
-			FVector Origin = FVector(UERebasedCoordinates.X, UERebasedCoordinates.Y, UERebasedCoordinates.Z);
-			TransformMatrix.SetOrigin(Origin);
+			TransformMatrix.SetOrigin(UEOrigin);
 
 			return FTransform(TransformMatrix);
 		}
