Fix invalid directory separator in GeoReferencing plugin so that Linux properly packages it

diff --git a/Engine/Plugins/Runtime/GeoReferencing/Source/GeoReferencing/GeoReferencing.Build.cs b/Engine/Plugins/Runtime/GeoReferencing/Source/GeoReferencing/GeoReferencing.Build.cs
--- a/Engine/Plugins/Runtime/GeoReferencing/Source/GeoReferencing/GeoReferencing.Build.cs
+++ b/Engine/Plugins/Runtime/GeoReferencing/Source/GeoReferencing/GeoReferencing.Build.cs
@@ -35,7 +35,7 @@ public class GeoReferencing : ModuleRules
 			}
 		);
 
-		string ProjRedistFolder = Path.Combine(PluginDirectory, @"Resources\PROJ\*");
+		string ProjRedistFolder = Path.Combine(PluginDirectory, @"Resources/PROJ/*");
 		RuntimeDependencies.Add(ProjRedistFolder, StagedFileType.UFS);
 	}
 }
