Code
- Fixed a sign error in GetPlanetCenterTransform
- Added a warning if ShaderPlatform settings enable rotation compression (can cause rendering accuracy issues with very large meshes)
- Added a warning if World bound checks is Enabled
- Removed 5.0 deprecated functions

Content
- Improved the Location Probe design with more precision at origin, and display Cardinal points axes
- Added better content for the Ellipsoid - Blue marble textures and materials, high resolution ellipsoid
- Cleanup of unused resources

#virtualized

diff --git a/Engine/Plugins/Runtime/GeoReferencing/Source/GeoReferencing/GeoReferencing.Build.cs b/Engine/Plugins/Runtime/GeoReferencing/Source/GeoReferencing/GeoReferencing.Build.cs
--- a/Engine/Plugins/Runtime/GeoReferencing/Source/GeoReferencing/GeoReferencing.Build.cs
+++ b/Engine/Plugins/Runtime/GeoReferencing/Source/GeoReferencing/GeoReferencing.Build.cs
@@ -30,8 +30,8 @@ public class GeoReferencing : ModuleRules
 				"Projects",
 				"Slate",
 				"SlateCore",
-				"SQLiteCore"
-
+				"SQLiteCore", 
+				"RHI"
 			}
 		);
 
diff --git a/Engine/Plugins/Runtime/GeoReferencing/Source/GeoReferencing/Private/CartesianCoordinates.cpp b/Engine/Plugins/Runtime/GeoReferencing/Source/GeoReferencing/Private/CartesianCoordinates.cpp
deleted file mode 100644
diff --git a/Engine/Plugins/Runtime/GeoReferencing/Source/GeoReferencing/Private/GeoReferencingModule.cpp b/Engine/Plugins/Runtime/GeoReferencing/Source/GeoReferencing/Private/GeoReferencingModule.cpp
--- a/Engine/Plugins/Runtime/GeoReferencing/Source/GeoReferencing/Private/GeoReferencingModule.cpp
+++ b/Engine/Plugins/Runtime/GeoReferencing/Source/GeoReferencing/Private/GeoReferencingModule.cpp
@@ -4,6 +4,7 @@
 #include "Interfaces/IPluginManager.h"
 #include "Modules/ModuleManager.h"
 #include "Styling/SlateStyleRegistry.h"
+#include "DataDrivenShaderPlatformInfo.h"
 
 DEFINE_LOG_CATEGORY(LogGeoReferencing);
 
@@ -23,6 +24,11 @@ void FGeoReferencingModule::StartupModule()
 	StyleSet->Set("ClassThumbnail.GeoReferencingSystem", new IMAGE_BRUSH("GeoReferencingSystem_64x", FVector2D(64.0f, 64.0f)));
 
 	FSlateStyleRegistry::RegisterSlateStyle(*StyleSet.Get());
+
+	if (FDataDrivenShaderPlatformInfo::GetSupportSceneDataCompressedTransforms(GMaxRHIShaderPlatform))
+	{
+		UE_LOG(LogGeoReferencing, Warning, TEXT("The engine is currently using Compressed Rotation Transforms. You may encounter accuracy issues with large neshes. We recommend setting all bSupportsSceneDataCompressedTransforms to false in the Engine/Config/[platform]/DataDrivenPlatformInfo.ini file"));
+	}
 }
 
 void FGeoReferencingModule::ShutdownModule()
diff --git a/Engine/Plugins/Runtime/GeoReferencing/Source/GeoReferencing/Private/GeoReferencingSystem.cpp b/Engine/Plugins/Runtime/GeoReferencing/Source/GeoReferencing/Private/GeoReferencingSystem.cpp
--- a/Engine/Plugins/Runtime/GeoReferencing/Source/GeoReferencing/Private/GeoReferencingSystem.cpp
+++ b/Engine/Plugins/Runtime/GeoReferencing/Source/GeoReferencing/Private/GeoReferencingSystem.cpp
@@ -2,7 +2,7 @@
 
 #include "GeoReferencingSystem.h"
 
-#include "CartesianCoordinates.h"
+#include "GameFramework/WorldSettings.h"
 #include "Interfaces/IPluginManager.h"
 #include "Ellipsoid.h"
 #include "Kismet/GameplayStatics.h"
@@ -151,6 +151,12 @@ void AGeoReferencingSystem::PostActorCreated()
 
 void AGeoReferencingSystem::Initialize()
 {
+	if (GetWorld() && GetWorld()->GetWorldSettings()->bEnableWorldBoundsChecks)
+	{
+		UE_LOG(LogGeoReferencing, Warning, TEXT("Enable World Bounds Checks is enabled in your World Settings. You might consider disabling it when working with large terrains, otherwise your pawns won't be able to go too far from the Origin"));
+	}
+
+
 	Impl = MakePimpl<FGeoReferencingSystemInternals>();
 
 	Impl->InitPROJLibrary();
@@ -178,98 +184,6 @@ void AGeoReferencingSystem::BeginDestroy()
 	}
 }
 
-#pragma region Old deprecated Prototypes
-
-void AGeoReferencingSystem::EngineToProjected(const FVector& EngineCoordinates, FCartesianCoordinates& ProjectedCoordinates)
-{
-	FVector Result;
-	EngineToProjected(EngineCoordinates, Result);
-	ProjectedCoordinates = FCartesianCoordinates(Result);
-}
-
-void AGeoReferencingSystem::ProjectedToEngine(const FCartesianCoordinates& ProjectedCoordinates, FVector& EngineCoordinates)
-{
-	ProjectedToEngine(ProjectedCoordinates.ToVector(), EngineCoordinates);
-}
-
-void AGeoReferencingSystem::EngineToECEF(const FVector& EngineCoordinates, FCartesianCoordinates& ECEFCoordinates)
-{
-	FVector Result;
-	EngineToECEF(EngineCoordinates, Result);
-	ECEFCoordinates = FCartesianCoordinates(Result);
-}
-
-void AGeoReferencingSystem::ECEFToEngine(const FCartesianCoordinates& ECEFCoordinates, FVector& EngineCoordinates)
-{
-	ECEFToEngine(ECEFCoordinates.ToVector(), EngineCoordinates);
-}
-
-
-void AGeoReferencingSystem::ProjectedToGeographic(const FCartesianCoordinates& ProjectedCoordinates, FGeographicCoordinates& GeographicCoordinates)
-{
-	ProjectedToGeographic(ProjectedCoordinates.ToVector(), GeographicCoordinates);
-}
-
-void AGeoReferencingSystem::GeographicToProjected(const FGeographicCoordinates& GeographicCoordinates, FCartesianCoordinates& ProjectedCoordinates)
-{
-	FVector Result;
-	GeographicToProjected(GeographicCoordinates, Result);
-	ProjectedCoordinates = FCartesianCoordinates(Result);
-}
-
-void AGeoReferencingSystem::ProjectedToECEF(const FCartesianCoordinates& ProjectedCoordinates, FCartesianCoordinates& ECEFCoordinates)
-{
-	FVector Result;
-	ProjectedToECEF(ProjectedCoordinates.ToVector(), Result);
-	ECEFCoordinates = FCartesianCoordinates(Result);
-}
-
-void AGeoReferencingSystem::ECEFToProjected(const FCartesianCoordinates& ECEFCoordinates, FCartesianCoordinates& ProjectedCoordinates)
-{
-	FVector Result;
-	ECEFToProjected(ECEFCoordinates.ToVector(), Result);
-	ProjectedCoordinates = FCartesianCoordinates(Result);
-}
-
-void AGeoReferencingSystem::GeographicToECEF(const FGeographicCoordinates& GeographicCoordinates, FCartesianCoordinates& ECEFCoordinates)
-{
-	FVector Result;
-	GeographicToECEF(GeographicCoordinates, Result);
-	ECEFCoordinates = FCartesianCoordinates(Result);
-}
-
-void AGeoReferencingSystem::ECEFToGeographic(const FCartesianCoordinates& ECEFCoordinates, FGeographicCoordinates& GeographicCoordinates)
-{
-	ECEFToGeographic(ECEFCoordinates.ToVector(), GeographicCoordinates);
-}
-
-void AGeoReferencingSystem::GetENUVectorsAtProjectedLocation(const FCartesianCoordinates& ProjectedCoordinates, FVector& East, FVector& North, FVector& Up)
-{
-	GetENUVectorsAtProjectedLocation(ProjectedCoordinates.ToVector(), East, North, Up);
-}
-
-void AGeoReferencingSystem::GetENUVectorsAtECEFLocation(const FCartesianCoordinates& ECEFCoordinates, FVector& East, FVector& North, FVector& Up)
-{
-	GetENUVectorsAtECEFLocation(ECEFCoordinates.ToVector(), East, North, Up);
-}
-
-void AGeoReferencingSystem::GetECEFENUVectorsAtECEFLocation(const FCartesianCoordinates& ECEFCoordinates, FVector& ECEFEast, FVector& ECEFNorth, FVector& ECEFUp)
-{
-	GetECEFENUVectorsAtECEFLocation(ECEFCoordinates.ToVector(), ECEFEast, ECEFNorth, ECEFUp);
-}
-
-FTransform AGeoReferencingSystem::GetTangentTransformAtProjectedLocation(const FCartesianCoordinates& ProjectedCoordinates)
-{
-	return GetTangentTransformAtProjectedLocation(ProjectedCoordinates.ToVector());
-}
-
-FTransform AGeoReferencingSystem::GetTangentTransformAtECEFLocation(const FCartesianCoordinates& ECEFCoordinates)
-{
-	return GetTangentTransformAtECEFLocation(ECEFCoordinates.ToVector());
-}
-#pragma endregion
-
-
 void AGeoReferencingSystem::EngineToProjected(const FVector& EngineCoordinates, FVector& ProjectedCoordinates)
 {
 	switch (PlanetShape)
@@ -655,7 +569,7 @@ FTransform AGeoReferencingSystem::GetPlanetCenterTransform()
 			// Don't go to transform yet, we must stay in double to apply the rebasing offset. 
 
 			// Get Origin, and convert UE units to meters, invert the Y coordinate because of left-handed UE Frame
-			FVector UEOrigin = TransformMatrix.GetOrigin() * FVector(100.0, -100.0, 100.0);
+			FVector UEOrigin = TransformMatrix.GetOrigin() * FVector(100.0, 100.0, 100.0);
 			TransformMatrix.SetOrigin(UEOrigin);
 
 			return FTransform(TransformMatrix);
diff --git a/Engine/Plugins/Runtime/GeoReferencing/Source/GeoReferencing/Public/CartesianCoordinates.h b/Engine/Plugins/Runtime/GeoReferencing/Source/GeoReferencing/Public/CartesianCoordinates.h
deleted file mode 100644
diff --git a/Engine/Plugins/Runtime/GeoReferencing/Source/GeoReferencing/Public/Ellipsoid.h b/Engine/Plugins/Runtime/GeoReferencing/Source/GeoReferencing/Public/Ellipsoid.h
--- a/Engine/Plugins/Runtime/GeoReferencing/Source/GeoReferencing/Public/Ellipsoid.h
+++ b/Engine/Plugins/Runtime/GeoReferencing/Source/GeoReferencing/Public/Ellipsoid.h
@@ -5,8 +5,6 @@
 #include "CoreMinimal.h"
 
 #include "GeographicCoordinates.h"
-#include "CartesianCoordinates.h"
-
 #include "Ellipsoid.generated.h"
 
 USTRUCT(BlueprintType)
diff --git a/Engine/Plugins/Runtime/GeoReferencing/Source/GeoReferencing/Public/GeoReferencingSystem.h b/Engine/Plugins/Runtime/GeoReferencing/Source/GeoReferencing/Public/GeoReferencingSystem.h
--- a/Engine/Plugins/Runtime/GeoReferencing/Source/GeoReferencing/Public/GeoReferencingSystem.h
+++ b/Engine/Plugins/Runtime/GeoReferencing/Source/GeoReferencing/Public/GeoReferencingSystem.h
@@ -9,7 +9,6 @@
 
 // Double precision structures
 #include "GeographicCoordinates.h"
-#include "CartesianCoordinates.h"
 #include "Ellipsoid.h"
 
 #include "GeoReferencingSystem.generated.h"
@@ -51,81 +50,11 @@ public:
 	UFUNCTION(BlueprintPure, Category = "GeoReferencing", meta = (WorldContext = "WorldContextObject"))
 	static AGeoReferencingSystem* GetGeoReferencingSystem(UObject* WorldContextObject);
 
-	// Deprecate all previous functions that were relying on Double precision placeholders
-
-#pragma region Old deprecated Prototypes
-
-	PRAGMA_DISABLE_DEPRECATION_WARNINGS
-
-	UE_DEPRECATED(5.0, "FCartesianCoordinates is deprecated : Use the version that uses a FVector instead.")
-	UFUNCTION(BlueprintCallable, Category = "GeoReferencing|Transformations", meta = (DeprecatedFunction, DeprecationMessage = "FCartesianCoordinates is deprecated : Use the version that uses a FVector instead."))
-	void EngineToProjected(const FVector& EngineCoordinates, FCartesianCoordinates& ProjectedCoordinates);
-
-	UE_DEPRECATED(5.0, "FCartesianCoordinates is deprecated : Use the version that uses a FVector instead.")
-	UFUNCTION(BlueprintCallable, Category = "GeoReferencing|Transformations", meta = (DeprecatedFunction, DeprecationMessage = "FCartesianCoordinates is deprecated : Use the version that uses a FVector instead."))
-	void ProjectedToEngine(const FCartesianCoordinates& ProjectedCoordinates, FVector& EngineCoordinates);
-
-	UE_DEPRECATED(5.0, "FCartesianCoordinates is deprecated : Use the version that uses a FVector instead.")
-	UFUNCTION(BlueprintCallable, Category = "GeoReferencing|Transformations", meta = (DeprecatedFunction, DeprecationMessage = "FCartesianCoordinates is deprecated : Use the version that uses a FVector instead."))
-	void EngineToECEF(const FVector& EngineCoordinates, FCartesianCoordinates& ECEFCoordinates);
-
-	UE_DEPRECATED(5.0, "FCartesianCoordinates is deprecated : Use the version that uses a FVector instead.")
-	UFUNCTION(BlueprintCallable, Category = "GeoReferencing|Transformations", meta = (DeprecatedFunction, DeprecationMessage = "FCartesianCoordinates is deprecated : Use the version that uses a FVector instead."))
-	void ECEFToEngine(const FCartesianCoordinates& ECEFCoordinates, FVector& EngineCoordinates);
-
-	UE_DEPRECATED(5.0, "FCartesianCoordinates is deprecated : Use the version that uses a FVector instead.")
-	UFUNCTION(BlueprintCallable, Category = "GeoReferencing|Transformations", meta = (DeprecatedFunction, DeprecationMessage = "FCartesianCoordinates is deprecated : Use the version that uses a FVector instead."))
-	void ProjectedToGeographic(const FCartesianCoordinates& ProjectedCoordinates, FGeographicCoordinates& GeographicCoordinates);
-
-	UE_DEPRECATED(5.0, "FCartesianCoordinates is deprecated : Use the version that uses a FVector instead.")
-	UFUNCTION(BlueprintCallable, Category = "GeoReferencing|Transformations", meta = (DeprecatedFunction, DeprecationMessage = "FCartesianCoordinates is deprecated : Use the version that uses a FVector instead."))
-	void GeographicToProjected(const FGeographicCoordinates& GeographicCoordinates, FCartesianCoordinates& ProjectedCoordinates);
-
-	UE_DEPRECATED(5.0, "FCartesianCoordinates is deprecated : Use the version that uses a FVector instead.")
-	UFUNCTION(BlueprintCallable, Category = "GeoReferencing|Transformations", meta = (DeprecatedFunction, DeprecationMessage = "FCartesianCoordinates is deprecated : Use the version that uses a FVector instead."))
-	void ProjectedToECEF(const FCartesianCoordinates& ProjectedCoordinates, FCartesianCoordinates& ECEFCoordinates);
-
-	UE_DEPRECATED(5.0, "FCartesianCoordinates is deprecated : Use the version that uses a FVector instead.")
-	UFUNCTION(BlueprintCallable, Category = "GeoReferencing|Transformations", meta = (DeprecatedFunction, DeprecationMessage = "FCartesianCoordinates is deprecated : Use the version that uses a FVector instead."))
-	void ECEFToProjected(const FCartesianCoordinates& ECEFCoordinates, FCartesianCoordinates& ProjectedCoordinates);
-
-	UE_DEPRECATED(5.0, "FCartesianCoordinates is deprecated : Use the version that uses a FVector instead.")
-	UFUNCTION(BlueprintCallable, Category = "GeoReferencing|Transformations", meta = (DeprecatedFunction, DeprecationMessage = "FCartesianCoordinates is deprecated : Use the version that uses a FVector instead."))
-	void GeographicToECEF(const FGeographicCoordinates& GeographicCoordinates, FCartesianCoordinates& ECEFCoordinates);
-
-	UE_DEPRECATED(5.0, "FCartesianCoordinates is deprecated : Use the version that uses a FVector instead.")
-	UFUNCTION(BlueprintCallable, Category = "GeoReferencing|Transformations", meta = (DeprecatedFunction, DeprecationMessage = "FCartesianCoordinates is deprecated : Use the version that uses a FVector instead."))
-	void ECEFToGeographic(const FCartesianCoordinates& ECEFCoordinates, FGeographicCoordinates& GeographicCoordinates);
-
-	UE_DEPRECATED(5.0, "FCartesianCoordinates is deprecated : Use the version that uses a FVector instead.")
-	UFUNCTION(BlueprintCallable, Category = "GeoReferencing|ENU", meta = (DeprecatedFunction, DeprecationMessage = "FCartesianCoordinates is deprecated : Use the version that uses a FVector instead."))
-	void GetENUVectorsAtProjectedLocation(const FCartesianCoordinates& ProjectedCoordinates, FVector& East, FVector& North, FVector& Up);
-
-	UE_DEPRECATED(5.0, "FCartesianCoordinates is deprecated : Use the version that uses a FVector instead.")
-	UFUNCTION(BlueprintCallable, Category = "GeoReferencing|ENU", meta = (DeprecatedFunction, DeprecationMessage = "FCartesianCoordinates is deprecated : Use the version that uses a FVector instead."))
-	void GetENUVectorsAtECEFLocation(const FCartesianCoordinates& ECEFCoordinates, FVector& East, FVector& North, FVector& Up);
-
-	UE_DEPRECATED(5.0, "FCartesianCoordinates is deprecated : Use the version that uses a FVector instead.")
-	UFUNCTION(BlueprintCallable, Category = "GeoReferencing|ENU", meta = (DeprecatedFunction, DeprecationMessage = "FCartesianCoordinates is deprecated : Use the version that uses a FVector instead."))
-	void GetECEFENUVectorsAtECEFLocation(const FCartesianCoordinates& ECEFCoordinates, FVector& ECEFEast, FVector& ECEFNorth, FVector& ECEFUp);
-
-	UE_DEPRECATED(5.0, "FCartesianCoordinates is deprecated : Use the version that uses a FVector instead.")
-	UFUNCTION(BlueprintCallable, Category = "GeoReferencing|TangentTransforms", meta = (DeprecatedFunction, DeprecationMessage = "FCartesianCoordinates is deprecated : Use the version that uses a FVector instead."))
-	FTransform GetTangentTransformAtProjectedLocation(const FCartesianCoordinates& ProjectedCoordinates);
-
-	UE_DEPRECATED(5.0, "FCartesianCoordinates is deprecated : Use the version that uses a FVector instead.")
-	UFUNCTION(BlueprintCallable, Category = "GeoReferencing|TangentTransforms", meta = (DeprecatedFunction, DeprecationMessage = "FCartesianCoordinates is deprecated : Use the version that uses a FVector instead."))
-	FTransform GetTangentTransformAtECEFLocation(const FCartesianCoordinates& ECEFCoordinates);
-
-	PRAGMA_ENABLE_DEPRECATION_WARNINGS
-
-#pragma endregion
-
 #pragma region New Prototypes for Blueprints
 	// We want to keep the same function names, but with a change in some argument types. UBT doesn't support that unless
 	// we create the new functions with a K2_ prefix and use the meta/Displayname tag to keep the same name. 
 	// To allow for a simple C++ usage too, we keep the C++ functions doing the actual work separately
-
+	// Now that we have used them for the BP Nodes, we can't move the UFUNCTION declaration to the actual new C++ function. Leave them here...
 
 	/**
 	* Convert a Vector expressed in ENGINE space to the PROJECTED CRS
diff --git a/Engine/Plugins/Runtime/GeoReferencing/Source/GeoReferencing/Public/GeographicCoordinates.h b/Engine/Plugins/Runtime/GeoReferencing/Source/GeoReferencing/Public/GeographicCoordinates.h
--- a/Engine/Plugins/Runtime/GeoReferencing/Source/GeoReferencing/Public/GeographicCoordinates.h
+++ b/Engine/Plugins/Runtime/GeoReferencing/Source/GeoReferencing/Public/GeographicCoordinates.h
@@ -66,29 +66,6 @@ class GEOREFERENCING_API UGeographicCoordinatesFunctionLibrary : public UBluepri
 		GeographicCoordinates.ToSeparateTexts(OutLatitude, OutLongitude, OutAltitude, IntegralDigitsLatLon, IntegralDigitsAlti, bAsDMS);
 	}
 
-	/**
-	 * Get the Coordinates as a float approximation.
-	 * USE WISELY as we can't guarantee there will no be rounding due to IEEE754 float encoding !
-	 **/
-	UE_DEPRECATED(5.0, "BP now support doubles, Function useless and can lead to precision issues")
-	UFUNCTION(BlueprintPure, Category = "GeoReferencing", meta = (DeprecatedFunction, DeprecationMessage = "BP now support doubles, Function useless and can lead to precision issues"))
-	static void ToFloatApproximation(UPARAM(ref) FGeographicCoordinates& GeographicCoordinates, float& OutLatitude, float& OutLongitude, float& OutAltitude)
-	{
-		OutLatitude  = static_cast<float>(GeographicCoordinates.Latitude);
-		OutLongitude = static_cast<float>(GeographicCoordinates.Longitude);
-		OutAltitude  = static_cast<float>(GeographicCoordinates.Altitude);
-	}
-
-	/**
-	 * Set the Coordinates from float approximation.
-	 * USE WISELY as we can't guarantee there will no be rounding due to IEEE754 float encoding !
-	 **/
-	UE_DEPRECATED(5.0, "BP now support doubles, Function useless and can lead to precision issues")
-	UFUNCTION(BlueprintPure, Category = "GeoReferencing|Coordinates", meta = (DeprecatedFunction, DeprecationMessage = "BP now support doubles, Function useless and can lead to precision issues"))
-	static FGeographicCoordinates MakeGeographicCoordinatesApproximation(const float& InLatitude, const float& InLongitude, const float& InAltitude)
-	{
-		return FGeographicCoordinates(static_cast<double>(InLongitude), static_cast<double>(InLatitude), static_cast<double>(InAltitude));
-	}
 
 
 	/**
