Changed the Log Level from Warning to Display to remove the Content Warning in our DDC CI/CD Tests

diff --git a/Engine/Plugins/Runtime/GeoReferencing/Source/GeoReferencing/Private/GeoReferencingModule.cpp b/Engine/Plugins/Runtime/GeoReferencing/Source/GeoReferencing/Private/GeoReferencingModule.cpp
--- a/Engine/Plugins/Runtime/GeoReferencing/Source/GeoReferencing/Private/GeoReferencingModule.cpp
+++ b/Engine/Plugins/Runtime/GeoReferencing/Source/GeoReferencing/Private/GeoReferencingModule.cpp
@@ -27,7 +27,7 @@ void FGeoReferencingModule::StartupModule()
 
 	if (FDataDrivenShaderPlatformInfo::GetSupportSceneDataCompressedTransforms(GMaxRHIShaderPlatform))
 	{
-		UE_LOG(LogGeoReferencing, Warning, TEXT("The engine is currently using Compressed Rotation Transforms. You may encounter accuracy issues with large neshes. We recommend setting all bSupportsSceneDataCompressedTransforms to false in the Engine/Config/[platform]/DataDrivenPlatformInfo.ini file"));
+		UE_LOG(LogGeoReferencing, Display, TEXT("The engine is currently using Compressed Rotation Transforms. You may encounter accuracy issues with large neshes. We recommend setting all bSupportsSceneDataCompressedTransforms to false in the Engine/Config/[platform]/DataDrivenPlatformInfo.ini file"));
 	}
 }
 
diff --git a/Engine/Plugins/Runtime/GeoReferencing/Source/GeoReferencing/Private/GeoReferencingSystem.cpp b/Engine/Plugins/Runtime/GeoReferencing/Source/GeoReferencing/Private/GeoReferencingSystem.cpp
--- a/Engine/Plugins/Runtime/GeoReferencing/Source/GeoReferencing/Private/GeoReferencingSystem.cpp
+++ b/Engine/Plugins/Runtime/GeoReferencing/Source/GeoReferencing/Private/GeoReferencingSystem.cpp
@@ -153,7 +153,7 @@ void AGeoReferencingSystem::Initialize()
 {
 	if (GetWorld() && GetWorld()->GetWorldSettings()->bEnableWorldBoundsChecks)
 	{
-		UE_LOG(LogGeoReferencing, Warning, TEXT("Enable World Bounds Checks is enabled in your World Settings. You might consider disabling it when working with large terrains, otherwise your pawns won't be able to go too far from the Origin"));
+		UE_LOG(LogGeoReferencing, Display, TEXT("Enable World Bounds Checks is enabled in your World Settings. You might consider disabling it when working with large terrains, otherwise your pawns won't be able to go too far from the Origin"));
 	}
 
 
