Fix missing file warning from GeoReferencing in non-editor builds

diff --git a/Engine/Plugins/Runtime/GeoReferencing/Source/GeoReferencing/GeoReferencing.Build.cs b/Engine/Plugins/Runtime/GeoReferencing/Source/GeoReferencing/GeoReferencing.Build.cs
--- a/Engine/Plugins/Runtime/GeoReferencing/Source/GeoReferencing/GeoReferencing.Build.cs
+++ b/Engine/Plugins/Runtime/GeoReferencing/Source/GeoReferencing/GeoReferencing.Build.cs
@@ -28,13 +28,16 @@ public class GeoReferencing : ModuleRules
 				"Engine",
 				"PROJ",
 				"Projects",
-				"Slate",
-				"SlateCore",
 				"SQLiteCore", 
 				"RHI"
 			}
 		);
 
+		if (Target.bCompileAgainstEditor)
+		{
+			PrivateDependencyModuleNames.Add("Slate");
+		}
+
 		string ProjRedistFolder = Path.Combine(PluginDirectory, @"Resources/PROJ/*");
 		RuntimeDependencies.Add(ProjRedistFolder, StagedFileType.UFS);
 	}
diff --git a/Engine/Plugins/Runtime/GeoReferencing/Source/GeoReferencing/Private/GeoReferencingModule.cpp b/Engine/Plugins/Runtime/GeoReferencing/Source/GeoReferencing/Private/GeoReferencingModule.cpp
--- a/Engine/Plugins/Runtime/GeoReferencing/Source/GeoReferencing/Private/GeoReferencingModule.cpp
+++ b/Engine/Plugins/Runtime/GeoReferencing/Source/GeoReferencing/Private/GeoReferencingModule.cpp
@@ -1,9 +1,7 @@
 // Copyright Epic Games, Inc. All Rights Reserved.
 
 #include "GeoReferencingModule.h"
-#include "Interfaces/IPluginManager.h"
 #include "Modules/ModuleManager.h"
-#include "Styling/SlateStyleRegistry.h"
 #include "DataDrivenShaderPlatformInfo.h"
 
 DEFINE_LOG_CATEGORY(LogGeoReferencing);
@@ -11,28 +9,10 @@ DEFINE_LOG_CATEGORY(LogGeoReferencing);
 IMPLEMENT_MODULE(FGeoReferencingModule, GeoReferencing)
 
 
-#define IMAGE_BRUSH(RelativePath, ...) FSlateImageBrush(StyleSet->RootToContentDir(RelativePath, TEXT(".png")), __VA_ARGS__)
-
 void FGeoReferencingModule::StartupModule()
 {
-	FString RootDir = IPluginManager::Get().FindPlugin(TEXT("GeoReferencing"))->GetBaseDir() / TEXT("Resources");
-
-	StyleSet = MakeUnique<FSlateStyleSet>(FName(TEXT("GeoReferencingStyle")));
-	StyleSet->SetContentRoot(RootDir);
-
-	StyleSet->Set("ClassIcon.GeoReferencingSystem", new IMAGE_BRUSH("GeoReferencingSystem_16x", FVector2D(16.0f, 16.0f)));
-	StyleSet->Set("ClassThumbnail.GeoReferencingSystem", new IMAGE_BRUSH("GeoReferencingSystem_64x", FVector2D(64.0f, 64.0f)));
-
-	FSlateStyleRegistry::RegisterSlateStyle(*StyleSet.Get());
-
 	if (FDataDrivenShaderPlatformInfo::GetSupportSceneDataCompressedTransforms(GMaxRHIShaderPlatform))
 	{
 		UE_LOG(LogGeoReferencing, Display, TEXT("The engine is currently using Compressed Rotation Transforms. You may encounter accuracy issues with large neshes. We recommend setting all bSupportsSceneDataCompressedTransforms to false in the Engine/Config/[platform]/DataDrivenPlatformInfo.ini file"));
 	}
 }
-
-void FGeoReferencingModule::ShutdownModule()
-{
-	FSlateStyleRegistry::UnRegisterSlateStyle(*StyleSet.Get());
-	StyleSet.Reset();
-}
diff --git a/Engine/Plugins/Runtime/GeoReferencing/Source/GeoReferencing/Public/GeoReferencingModule.h b/Engine/Plugins/Runtime/GeoReferencing/Source/GeoReferencing/Public/GeoReferencingModule.h
--- a/Engine/Plugins/Runtime/GeoReferencing/Source/GeoReferencing/Public/GeoReferencingModule.h
+++ b/Engine/Plugins/Runtime/GeoReferencing/Source/GeoReferencing/Public/GeoReferencingModule.h
@@ -4,7 +4,6 @@
 
 #include "Modules/ModuleInterface.h"
 #include "CoreMinimal.h"
-#include "Styling/SlateStyle.h"
 
 GEOREFERENCING_API DECLARE_LOG_CATEGORY_EXTERN(LogGeoReferencing, Log, All);
 
@@ -17,8 +16,4 @@ public:
 	FGeoReferencingModule() = default;
 
 	virtual void StartupModule() override;
-	virtual void ShutdownModule() override;
-
-private:
-	TUniquePtr<FSlateStyleSet> StyleSet;
 };
\ No newline at end of file
diff --git a/Engine/Plugins/Runtime/GeoReferencing/Source/GeoReferencingEditor/GeoReferencingEditor.Build.cs b/Engine/Plugins/Runtime/GeoReferencing/Source/GeoReferencingEditor/GeoReferencingEditor.Build.cs
--- a/Engine/Plugins/Runtime/GeoReferencing/Source/GeoReferencingEditor/GeoReferencingEditor.Build.cs
+++ b/Engine/Plugins/Runtime/GeoReferencing/Source/GeoReferencingEditor/GeoReferencingEditor.Build.cs
@@ -12,15 +12,17 @@ public class GeoReferencingEditor : ModuleRules
 	{
 		PCHUsage = ModuleRules.PCHUsageMode.UseExplicitOrSharedPCHs;
 
-		PublicDependencyModuleNames.AddRange(
+		PrivateDependencyModuleNames.AddRange(
 			new string[]
 			{
 				"Core",
 				"CoreUObject",
 				"Engine",
-				"InputCore",
 				"GeoReferencing",
+				"InputCore",
 				"Projects",
+				"Slate",
+				"SlateCore",
 				"UnrealEd"
 			}
 		);
diff --git a/Engine/Plugins/Runtime/GeoReferencing/Source/GeoReferencingEditor/Private/GeoReferencingEditor.cpp b/Engine/Plugins/Runtime/GeoReferencing/Source/GeoReferencingEditor/Private/GeoReferencingEditor.cpp
--- a/Engine/Plugins/Runtime/GeoReferencing/Source/GeoReferencingEditor/Private/GeoReferencingEditor.cpp
+++ b/Engine/Plugins/Runtime/GeoReferencing/Source/GeoReferencingEditor/Private/GeoReferencingEditor.cpp
@@ -1,11 +1,36 @@
 // Copyright Epic Games, Inc. All Rights Reserved.
 
+#include "Interfaces/IPluginManager.h"
 #include "Modules/ModuleManager.h"
+#include "Styling/SlateStyleRegistry.h"
+#include "Styling/SlateStyle.h"
+
+#define IMAGE_BRUSH(RelativePath, ...) FSlateImageBrush(StyleSet->RootToContentDir(RelativePath, TEXT(".png")), __VA_ARGS__)
 
-#include "Modules/ModuleManager.h"
 
 class FGeoReferencingEditorModule : public IModuleInterface
 {
+	TUniquePtr<FSlateStyleSet> StyleSet;
+
+
+	virtual void StartupModule() override
+	{
+		FString RootDir = IPluginManager::Get().FindPlugin(TEXT("GeoReferencing"))->GetBaseDir() / TEXT("Resources");
+
+		StyleSet = MakeUnique<FSlateStyleSet>(FName(TEXT("GeoReferencingStyle")));
+		StyleSet->SetContentRoot(RootDir);
+
+		StyleSet->Set("ClassIcon.GeoReferencingSystem", new IMAGE_BRUSH("GeoReferencingSystem_16x", FVector2D(16.0f, 16.0f)));
+		StyleSet->Set("ClassThumbnail.GeoReferencingSystem", new IMAGE_BRUSH("GeoReferencingSystem_64x", FVector2D(64.0f, 64.0f)));
+
+		FSlateStyleRegistry::RegisterSlateStyle(*StyleSet.Get());
+	}
+
+	virtual void ShutdownModule() override
+	{
+		FSlateStyleRegistry::UnRegisterSlateStyle(*StyleSet.Get());
+		StyleSet.Reset();
+	}
 };
 
 IMPLEMENT_MODULE(FGeoReferencingEditorModule, GeoReferencingEditor)
